[NNE] NNERuntimeORT Perf fix

Problem:
* NNERuntimeORT was using an ORT session utilizing only 2 threads limiting performance.

Solution:
* Let ORT use it's default thread configuration. Also add a performance scope for easy tracking.

Test:
* Win64 Editor: style transfer + all test
* Win64 Standalone: style transfer + all test

#rnx

diff --git a/Engine/Plugins/NNERuntimeORT/Source/NNERuntimeORT/Private/NNERuntimeORTModel.cpp b/Engine/Plugins/NNERuntimeORT/Source/NNERuntimeORT/Private/NNERuntimeORTModel.cpp
--- a/Engine/Plugins/NNERuntimeORT/Source/NNERuntimeORT/Private/NNERuntimeORTModel.cpp
+++ b/Engine/Plugins/NNERuntimeORT/Source/NNERuntimeORT/Private/NNERuntimeORTModel.cpp
@@ -1,6 +1,8 @@
 // Copyright Epic Games, Inc. All Rights Reserved.
 
 #include "NNERuntimeORTModel.h"
+
+#include "HAL/PlatformMisc.h"
 #include "NNERuntimeORT.h"
 #include "NNERuntimeORTUtils.h"
 
@@ -86,7 +88,6 @@ namespace UE::NNERuntimeORT::Private
 		AllocatorInfo = MakeUnique<Ort::MemoryInfo>(Ort::MemoryInfo::CreateCpu(OrtDeviceAllocator, OrtMemTypeCPU));
 		SessionOptions = MakeUnique<Ort::SessionOptions>();
 
-		SessionOptions->SetIntraOpNumThreads(RuntimeConf.NumberOfThreads);
 		SessionOptions->SetGraphOptimizationLevel(RuntimeConf.OptimizationLevel);
 
 		return true;
@@ -180,6 +181,8 @@ namespace UE::NNERuntimeORT::Private
 	{
 		checkf(Session.IsValid(), TEXT("FModelInstanceORT::RunSync(): Called without a Session, FModelInstanceORT::Init() should have been called."));
 
+		SCOPED_NAMED_EVENT_TEXT("FModelInstanceORTBase::RunSync", FColor::Magenta);
+
 		// Verify the model inputs were prepared
 		if (NNE::Internal::FModelInstanceBase<ModelInterface>::InputTensorShapes.Num() == 0)
 		{
diff --git a/Engine/Plugins/NNERuntimeORT/Source/NNERuntimeORT/Private/NNERuntimeORTModel.h b/Engine/Plugins/NNERuntimeORT/Source/NNERuntimeORT/Private/NNERuntimeORTModel.h
--- a/Engine/Plugins/NNERuntimeORT/Source/NNERuntimeORT/Private/NNERuntimeORTModel.h
+++ b/Engine/Plugins/NNERuntimeORT/Source/NNERuntimeORT/Private/NNERuntimeORTModel.h
@@ -15,9 +15,7 @@ namespace UE::NNERuntimeORT::Private
 {
 	struct FRuntimeConf
 	{
-		uint32 NumberOfThreads = 2;
 		GraphOptimizationLevel OptimizationLevel = GraphOptimizationLevel::ORT_ENABLE_ALL;
-		EThreadPriority ThreadPriority = EThreadPriority::TPri_Normal;
 	};
 
 	template <class ModelInterface, class TensorBinding> class FModelInstanceORTBase : public NNE::Internal::FModelInstanceBase<ModelInterface>
