[NNE] NNERuntimeORT Load DirectML DLL on windows for the ORTDml runtime

Problem:
* NNERuntimeORT Dml need the DirectML dll to be loaded on windows but was not doing it itself, meaning NNERuntimeRDG was needed to be enabled for NNERuntimeORT Dml to be working.

Solution:
* Load required DirectML dll on NNERuntimeORT startup. This is fine even if NNERuntimeRDG is loading the DLL itself as DLL are ref counted per process and unloaded on process exist in all case.

Test:
* Win64 Editor: in house model tests (without NNERuntimeRDG) + regular all tests
* Win64 Standalone: all tests

#rnx

diff --git a/Engine/Plugins/NNE/NNERuntimeORT/Source/NNERuntimeORT/Private/NNERuntimeORTModule.cpp b/Engine/Plugins/NNE/NNERuntimeORT/Source/NNERuntimeORT/Private/NNERuntimeORTModule.cpp
--- a/Engine/Plugins/NNE/NNERuntimeORT/Source/NNERuntimeORT/Private/NNERuntimeORTModule.cpp
+++ b/Engine/Plugins/NNE/NNERuntimeORT/Source/NNERuntimeORT/Private/NNERuntimeORTModule.cpp
@@ -33,7 +33,8 @@ namespace UE::NNERuntimeORT::Private::DllHelper
 		DllHandles.Add(DllHandle);
 		return true;
 	}
-	}
+
+} // namespace UE::NNERuntimeRDG::Private::Dml
 
 void FNNERuntimeORTModule::StartupModule()
 {
@@ -46,20 +47,34 @@ void FNNERuntimeORTModule::StartupModule()
 		return;
 	}
 
+#if PLATFORM_WINDOWS
+	const FString ModuleDir = FPlatformProcess::GetModulesDirectory();
+	const FString DirectMLSharedLibPath = FPaths::Combine(ModuleDir, TEXT(PREPROCESSOR_TO_STRING(DIRECTML_PATH)), TEXT("DirectML.dll"));
+
+	const bool bDirectMLDllLoaded = UE::NNERuntimeORT::Private::DllHelper::GetDllHandle(DirectMLSharedLibPath, DllHandles);
+	if (!bDirectMLDllLoaded)
+	{
+		UE_LOG(LogNNE, Error, TEXT("Failed to load DirectML shared library. ORT Dml Runtime won't be available."));
+	}
+#endif // PLATFORM_WINDOWS
+
 	Ort::InitApi();
 
 #if PLATFORM_WINDOWS
-	// NNE runtime ORT Dml startup
-	NNERuntimeORTDml = NewObject<UNNERuntimeORTDml>();
-	if (NNERuntimeORTDml.IsValid())
+	if (bDirectMLDllLoaded)
 	{
-		TWeakInterfacePtr<INNERuntime> RuntimeDmlInterface(NNERuntimeORTDml.Get());
+		// NNE runtime ORT Dml startup
+		NNERuntimeORTDml = NewObject<UNNERuntimeORTDml>();
+		if (NNERuntimeORTDml.IsValid())
+		{
+			TWeakInterfacePtr<INNERuntime> RuntimeDmlInterface(NNERuntimeORTDml.Get());
 
-		NNERuntimeORTDml->Init();
-		NNERuntimeORTDml->AddToRoot();
-		UE::NNE::RegisterRuntime(RuntimeDmlInterface);
+			NNERuntimeORTDml->Init();
+			NNERuntimeORTDml->AddToRoot();
+			UE::NNE::RegisterRuntime(RuntimeDmlInterface);
+		}
 	}
-#endif
+#endif // PLATFORM_WINDOWS
 
 	// NNE runtime ORT Cpu startup
 	NNERuntimeORTCpu = NewObject<UNNERuntimeORTCpu>();
