[NNE] NNERuntimeRDG model convertion from onnx fix.

Problem:
NNERuntimeRDG was not able to cook certain onnx models following adoption of the onnx lib to convert from ONNX model to NNE.

Solution:
proto onnx c++ api is for some model returning tensors with initializer in the input list (onnx python api would not return those),
thus NNEUtitilitiesModelOptimizerNNE need to filter those out as NNE expect only the user defined tensor in the input list.

Test:
Win64 Editor: all tests
Win64 Standalone: all tests

#rnx

diff --git a/Engine/Plugins/Experimental/NNERuntimeRDG/Source/NNERuntimeRDG/Private/Dml/NNEDmlModel.cpp b/Engine/Plugins/Experimental/NNERuntimeRDG/Source/NNERuntimeRDG/Private/Dml/NNEDmlModel.cpp
--- a/Engine/Plugins/Experimental/NNERuntimeRDG/Source/NNERuntimeRDG/Private/Dml/NNEDmlModel.cpp
+++ b/Engine/Plugins/Experimental/NNERuntimeRDG/Source/NNERuntimeRDG/Private/Dml/NNEDmlModel.cpp
@@ -44,7 +44,7 @@ FModelInfo* FModelInfo::Get()
 
 FModelInfo::FModelInfo()
 	: Guid((int32)'R', (int32)'D', (int32)'G', (int32)'D')
-	, Version(0x3)
+	, Version(0x4)
 {
 }
 
diff --git a/Engine/Plugins/Experimental/NNERuntimeRDG/Source/NNERuntimeRDG/Private/NNERuntimeRDGHlsl.cpp b/Engine/Plugins/Experimental/NNERuntimeRDG/Source/NNERuntimeRDG/Private/NNERuntimeRDGHlsl.cpp
--- a/Engine/Plugins/Experimental/NNERuntimeRDG/Source/NNERuntimeRDG/Private/NNERuntimeRDGHlsl.cpp
+++ b/Engine/Plugins/Experimental/NNERuntimeRDG/Source/NNERuntimeRDG/Private/NNERuntimeRDGHlsl.cpp
@@ -42,7 +42,7 @@
 using namespace UE::NNERuntimeRDG::Private::Hlsl;
 
 FGuid UNNERuntimeRDGHlslImpl::GUID = FGuid((int32)'R', (int32)'D', (int32)'G', (int32)'H');
-int32 UNNERuntimeRDGHlslImpl::Version = 0x00000003;
+int32 UNNERuntimeRDGHlslImpl::Version = 0x00000004;
 
 bool UNNERuntimeRDGHlslImpl::Init()
 {
diff --git a/Engine/Plugins/NNE/NNERuntimeORT/Source/NNEUtilities/Private/NNEUtilitiesModelOptimizerNNE.cpp b/Engine/Plugins/NNE/NNERuntimeORT/Source/NNEUtilities/Private/NNEUtilitiesModelOptimizerNNE.cpp
--- a/Engine/Plugins/NNE/NNERuntimeORT/Source/NNEUtilities/Private/NNEUtilitiesModelOptimizerNNE.cpp
+++ b/Engine/Plugins/NNE/NNERuntimeORT/Source/NNEUtilities/Private/NNEUtilitiesModelOptimizerNNE.cpp
@@ -221,6 +221,13 @@ namespace ModelOptimizerNNEHelper
 		// Add tensors for graph inputs
 		for (const onnx::ValueInfoProto& Input : Graph.input())
 		{
+			// ONNX GraphProto sometime return initializers as input,
+			// we skip them here as we only want user providable inputs as NNE inputs
+			if (GetInitializerFromGraphProto(Graph, Input.name())  != nullptr)
+			{
+				continue;
+			}
+
 			ENNETensorDataType DataType;
 			TArray<int32> Shape;
 			GetTensorInfoFromONNXValueInfo(Input, Shape, DataType);
