[NNE] NNERuntimeORTDml throw error when initializing on headless servers.

Problem:
If the plugin is enabled on hardware without DX12 NNERuntimeORTDml runtime throw error at init. We would rather not have the runtime registered at all in that case.

Solution:
Only register the runtime on DX12 capable hardware.

Test: Win64 Editor: smoke tests

#rnx

diff --git a/Engine/Plugins/Experimental/NNEDenoiser/Source/NNEDenoiser/Private/NNEDenoiserViewExtension.cpp b/Engine/Plugins/Experimental/NNEDenoiser/Source/NNEDenoiser/Private/NNEDenoiserViewExtension.cpp
--- a/Engine/Plugins/Experimental/NNEDenoiser/Source/NNEDenoiser/Private/NNEDenoiserViewExtension.cpp
+++ b/Engine/Plugins/Experimental/NNEDenoiser/Source/NNEDenoiser/Private/NNEDenoiserViewExtension.cpp
@@ -40,7 +40,7 @@ static TAutoConsoleVariable<int32> CVarNNEDenoiserModelData(
 
 static TAutoConsoleVariable<int32> CVarNNEDenoiserRuntimeType(
 	TEXT("NNEDenoiser.Runtime.Type"),
-	1,
+	0,
 	TEXT("Defines the runtime type to run the denoiser model.\n")
 	TEXT("  0: CPU\n")
 	TEXT("  1: GPU\n")
diff --git a/Engine/Plugins/NNE/NNERuntimeORT/Source/NNERuntimeORT/Private/NNERuntimeORT.cpp b/Engine/Plugins/NNE/NNERuntimeORT/Source/NNERuntimeORT/Private/NNERuntimeORT.cpp
--- a/Engine/Plugins/NNE/NNERuntimeORT/Source/NNERuntimeORT/Private/NNERuntimeORT.cpp
+++ b/Engine/Plugins/NNE/NNERuntimeORT/Source/NNERuntimeORT/Private/NNERuntimeORT.cpp
@@ -12,6 +12,10 @@
 #include "NNERuntimeORTUtils.h"
 #include "NNEUtilitiesModelOptimizer.h"
 
+#if PLATFORM_WINDOWS
+#include "ID3D12DynamicRHI.h"
+#endif //PLATFORM_WINDOWS
+
 #include UE_INLINE_GENERATED_CPP_BY_NAME(NNERuntimeORT)
 
 FGuid UNNERuntimeORTDml::GUID = FGuid((int32)'O', (int32)'G', (int32)'P', (int32)'U');
@@ -102,6 +106,16 @@ void UNNERuntimeORTDml::Init()
 	ORTEnvironment = MakeShared<Ort::Env>();
 }
 
+bool UNNERuntimeORTDml::IsAvailable()
+{
+#if PLATFORM_WINDOWS
+	// In order to use DirectML we need D3D12
+	return IsRHID3D12();
+#else
+	return false;
+#endif
+}
+
 void UNNERuntimeORTCpu::Init()
 {
 	check(!ORTEnvironment.IsValid());
diff --git a/Engine/Plugins/NNE/NNERuntimeORT/Source/NNERuntimeORT/Private/NNERuntimeORT.h b/Engine/Plugins/NNE/NNERuntimeORT/Source/NNERuntimeORT/Private/NNERuntimeORT.h
--- a/Engine/Plugins/NNE/NNERuntimeORT/Source/NNERuntimeORT/Private/NNERuntimeORT.h
+++ b/Engine/Plugins/NNE/NNERuntimeORT/Source/NNERuntimeORT/Private/NNERuntimeORT.h
@@ -30,6 +30,7 @@ public:
 	UNNERuntimeORTDml() {};
 	virtual ~UNNERuntimeORTDml() {}
 
+	bool IsAvailable();
 	void Init();
 
 	virtual FString GetRuntimeName() const override;
diff --git a/Engine/Plugins/NNE/NNERuntimeORT/Source/NNERuntimeORT/Private/NNERuntimeORTModule.cpp b/Engine/Plugins/NNE/NNERuntimeORT/Source/NNERuntimeORT/Private/NNERuntimeORTModule.cpp
--- a/Engine/Plugins/NNE/NNERuntimeORT/Source/NNERuntimeORT/Private/NNERuntimeORTModule.cpp
+++ b/Engine/Plugins/NNE/NNERuntimeORT/Source/NNERuntimeORT/Private/NNERuntimeORTModule.cpp
@@ -65,7 +65,7 @@ void FNNERuntimeORTModule::StartupModule()
 	{
 		// NNE runtime ORT Dml startup
 		NNERuntimeORTDml = NewObject<UNNERuntimeORTDml>();
-		if (NNERuntimeORTDml.IsValid())
+		if (NNERuntimeORTDml.IsValid() && NNERuntimeORTDml->IsAvailable())
 		{
 			TWeakInterfacePtr<INNERuntime> RuntimeDmlInterface(NNERuntimeORTDml.Get());
 
