[NNE] NNERuntimeORTDml always initialize runtime but allow clients to know if current platform cannot create model.

Problem:
CL 30082211 introduced a regression while fixing logging errors when a model was create on a system where DX12 is not available. The fix was to not register the runtime when DX12 was not available, however runtime are responsible for both cooking and creating models and the cooking process is using a null RHI. So we need the runtime on non DX12 context for cooking purpose.

Solution:
Always register the runtime & make CanCreateModel() return false on non dx12 system, so customers can query this if they dont want the logs while calling CreateModelData on non DX12 systems.

Test: Win64 Editor: smoke tests
Test: Win64 Standalone: all tests

#rnx

diff --git a/Engine/Plugins/NNE/NNERuntimeORT/Source/NNERuntimeORT/Private/NNERuntimeORT.cpp b/Engine/Plugins/NNE/NNERuntimeORT/Source/NNERuntimeORT/Private/NNERuntimeORT.cpp
--- a/Engine/Plugins/NNE/NNERuntimeORT/Source/NNERuntimeORT/Private/NNERuntimeORT.cpp
+++ b/Engine/Plugins/NNE/NNERuntimeORT/Source/NNERuntimeORT/Private/NNERuntimeORT.cpp
@@ -106,16 +106,6 @@ void UNNERuntimeORTDml::Init()
 	ORTEnvironment = MakeShared<Ort::Env>();
 }
 
-bool UNNERuntimeORTDml::IsAvailable()
-{
-#if PLATFORM_WINDOWS
-	// In order to use DirectML we need D3D12
-	return IsRHID3D12();
-#else
-	return false;
-#endif
-}
-
 void UNNERuntimeORTCpu::Init()
 {
 	check(!ORTEnvironment.IsValid());
@@ -191,6 +181,12 @@ bool UNNERuntimeORTDml::CanCreateModelGPU(TObjectPtr<UNNEModelData> ModelData) c
 {
 	check(ModelData != nullptr);
 
+	// In order to use DirectML we need D3D12
+	if (!IsRHID3D12())
+	{
+		return false;
+	}
+
 	constexpr int32 GuidSize = sizeof(UNNERuntimeORTDml::GUID);
 	constexpr int32 VersionSize = sizeof(UNNERuntimeORTDml::Version);
 	const TSharedPtr<UE::NNE::FSharedModelData> SharedData = ModelData->GetModelData(GetRuntimeName());
diff --git a/Engine/Plugins/NNE/NNERuntimeORT/Source/NNERuntimeORT/Private/NNERuntimeORT.h b/Engine/Plugins/NNE/NNERuntimeORT/Source/NNERuntimeORT/Private/NNERuntimeORT.h
--- a/Engine/Plugins/NNE/NNERuntimeORT/Source/NNERuntimeORT/Private/NNERuntimeORT.h
+++ b/Engine/Plugins/NNE/NNERuntimeORT/Source/NNERuntimeORT/Private/NNERuntimeORT.h
@@ -30,7 +30,6 @@ public:
 	UNNERuntimeORTDml() {};
 	virtual ~UNNERuntimeORTDml() {}
 
-	bool IsAvailable();
 	void Init();
 
 	virtual FString GetRuntimeName() const override;
diff --git a/Engine/Plugins/NNE/NNERuntimeORT/Source/NNERuntimeORT/Private/NNERuntimeORTModule.cpp b/Engine/Plugins/NNE/NNERuntimeORT/Source/NNERuntimeORT/Private/NNERuntimeORTModule.cpp
--- a/Engine/Plugins/NNE/NNERuntimeORT/Source/NNERuntimeORT/Private/NNERuntimeORTModule.cpp
+++ b/Engine/Plugins/NNE/NNERuntimeORT/Source/NNERuntimeORT/Private/NNERuntimeORTModule.cpp
@@ -65,7 +65,7 @@ void FNNERuntimeORTModule::StartupModule()
 	{
 		// NNE runtime ORT Dml startup
 		NNERuntimeORTDml = NewObject<UNNERuntimeORTDml>();
-		if (NNERuntimeORTDml.IsValid() && NNERuntimeORTDml->IsAvailable())
+		if (NNERuntimeORTDml.IsValid())
 		{
 			TWeakInterfacePtr<INNERuntime> RuntimeDmlInterface(NNERuntimeORTDml.Get());
 
