[NNE] NNERuntimeRDG support int64 to int32 attribute convertion.

Problem:
Some onnx operator (Slice opset 9 for ex) have int64 attributes witch can be used as max_int64 to express infinity. At the same time the RDG internal format have those attribute converted to int32 and currently fail complaining at an overflow in those case.

Solution:
When converting from int64 to int32 clamp to int32 range and display a log. Allowing to carry infinity meaning from int64 to int32.

Test: Win64 Editor: smoke tests + added related model test
Test: Win64 Standalone: all tests

#rnx

#virtualized

diff --git a/Engine/Plugins/NNE/NNERuntimeORT/Source/NNEUtilities/Private/NNEUtilitiesModelOptimizerNNE.cpp b/Engine/Plugins/NNE/NNERuntimeORT/Source/NNEUtilities/Private/NNEUtilitiesModelOptimizerNNE.cpp
--- a/Engine/Plugins/NNE/NNERuntimeORT/Source/NNEUtilities/Private/NNEUtilitiesModelOptimizerNNE.cpp
+++ b/Engine/Plugins/NNE/NNERuntimeORT/Source/NNEUtilities/Private/NNEUtilitiesModelOptimizerNNE.cpp
@@ -322,13 +322,13 @@ namespace ModelOptimizerNNEHelper
 					Values.Reserve(Attribute.ints_size());
 					for (int64 Value64 : Attribute.ints())
 					{
-						int32 Value = (int32)Value64;
-						if ((int64)Value != Value64)
+						int64 ValueClamped = FMath::Clamp<int64>(Value64, MIN_int32, MAX_int32);
+						if (ValueClamped != Value64)
 						{
-							UE_LOG(LogNNE, Error, TEXT("Overflow detected when converting to int32 attribute '%s' in node '%s'"), *AttributeName, *NNEOpType);
-							return false;
+							UE_LOG(LogNNE, Display, TEXT("Overflow detected when converting to int32 attribute '%s' in node '%s'"), *AttributeName, *NNEOpType);
 						}
-						Values.Add(Value);
+						
+						Values.Add((int32)ValueClamped);
 					}
 
 					Builder->AddOperatorAttribute(Op, AttributeName, FNNEAttributeValue(Values));
