[NNE] NNERuntimeORT Add a cvar to enable ORT profiling mode.

Problems:
* As ORT is running in it's own context we can't use UE tooling to profile model runned by it.

Solutions:
* Add a cvar to enable ORT profiling mode witch an be inspected using recommended ORT tooling.

Test: Win64 Editor: smoke tests + ORT profiling

#rnx

diff --git a/Engine/Plugins/NNE/NNERuntimeORT/Source/NNERuntimeORT/Private/NNERuntimeORTModel.cpp b/Engine/Plugins/NNE/NNERuntimeORT/Source/NNERuntimeORT/Private/NNERuntimeORTModel.cpp
--- a/Engine/Plugins/NNE/NNERuntimeORT/Source/NNERuntimeORT/Private/NNERuntimeORTModel.cpp
+++ b/Engine/Plugins/NNE/NNERuntimeORT/Source/NNERuntimeORT/Private/NNERuntimeORTModel.cpp
@@ -81,6 +81,16 @@ namespace UE::NNERuntimeORT::Private
 		return true;
 	}
 
+	static int32 ORTProfilingSessionNumber = 0;
+	static TAutoConsoleVariable<bool> CVarNNERuntimeORTEnableProfiling(
+		TEXT("nne.ort.enableprofiling"),
+		false,
+		TEXT("True if NNERuntimeORT plugin should create ORT sessions with profiling enabled.\n")
+		TEXT("When profiling is enabled ORT will create standard performance tracing json files next to the editor executable.\n")
+		TEXT("The files will be prefixed by 'NNERuntimeORTProfile_' and can be loaded for example using chrome://tracing.\n")
+		TEXT("More information can be found at https://onnxruntime.ai/docs/performance/tune-performance/profiling-tools.html\n"),
+		ECVF_Default);
+
 	template <class ModelInterface, class TensorBinding> 
 	bool FModelInstanceORTBase<ModelInterface, TensorBinding>::InitializedAndConfigureMembers()
 	{
@@ -88,7 +98,19 @@ namespace UE::NNERuntimeORT::Private
 		AllocatorInfo = MakeUnique<Ort::MemoryInfo>(Ort::MemoryInfo::CreateCpu(OrtDeviceAllocator, OrtMemTypeCPU));
 		SessionOptions = MakeUnique<Ort::SessionOptions>();
 
+		
 		SessionOptions->SetGraphOptimizationLevel(RuntimeConf.OptimizationLevel);
+		if (CVarNNERuntimeORTEnableProfiling.GetValueOnGameThread())
+		{
+			FString ProfilingFilePrefix("NNERuntimeORTProfile_");
+			ProfilingFilePrefix += FString::FromInt(ORTProfilingSessionNumber);
+			++ORTProfilingSessionNumber;
+			#if PLATFORM_WINDOWS
+				SessionOptions->EnableProfiling(*ProfilingFilePrefix);
+			#else
+				SessionOptions->EnableProfiling(TCHAR_TO_ANSI(*ProfilingFilePrefix));
+			#endif
+		}
 
 		return true;
 	}
