[NNE] Fix missing include file for mac non-unity build.

Test: Win64 Editor: smoke tests

#rnx

diff --git a/Engine/Plugins/NNE/NNERuntimeORT/Source/NNERuntimeORT/Private/NNERuntimeORTModel.cpp b/Engine/Plugins/NNE/NNERuntimeORT/Source/NNERuntimeORT/Private/NNERuntimeORTModel.cpp
--- a/Engine/Plugins/NNE/NNERuntimeORT/Source/NNERuntimeORT/Private/NNERuntimeORTModel.cpp
+++ b/Engine/Plugins/NNE/NNERuntimeORT/Source/NNERuntimeORT/Private/NNERuntimeORTModel.cpp
@@ -2,6 +2,7 @@
 
 #include "NNERuntimeORTModel.h"
 
+#include "HAL/IConsoleManager.h"
 #include "HAL/PlatformMisc.h"
 #include "NNERuntimeORT.h"
 #include "NNERuntimeORTUtils.h"
