[NNE] Add opset versioning to NNEQA
Problem
Versioned operators require proper testing infrastructure.

Solution
Complement NNEQA json test description format with syntax to express information regarding which versions of the operators to test and how.
Introduce a new type of tests (opset tests) for testing all the supported versions of the ops. Contrarily to iteration tests (those we had so far, testing a single version) they are not supposed to be run on every development iteration, but with a much lower frequency, due to their length.
Through the newly introduced `default_version` and `versions` fields of a test target it's possible to, respectively, indicate which version of the op should be tested in iteration tests and which versions are supported overall (tested in opset tests). Both are optional fields.
`versions` is also a field of input-output sets and attribute sets, restricting the pool of versions the corresponding tests apply to.

Note: opset tests have been added only for a few OPs here. Another CL should add opset tests for the remaining OPs.

Tests
- NNEIncubatorEditor Win64
- Non-unity compile

#rnx

diff --git a/Engine/Plugins/Experimental/NNERuntimeRDG/Source/NNERuntimeRDG/Private/Dml/NNEDmlOperatorPool.cpp b/Engine/Plugins/Experimental/NNERuntimeRDG/Source/NNERuntimeRDG/Private/Dml/NNEDmlOperatorPool.cpp
--- a/Engine/Plugins/Experimental/NNERuntimeRDG/Source/NNERuntimeRDG/Private/Dml/NNEDmlOperatorPool.cpp
+++ b/Engine/Plugins/Experimental/NNERuntimeRDG/Source/NNERuntimeRDG/Private/Dml/NNEDmlOperatorPool.cpp
@@ -320,6 +320,7 @@ NNE_DML_REGISTER_POOLING_OP(AveragePool, AVERAGE, false, 1)
 NNE_DML_REGISTER_POOLING_OP(AveragePool, AVERAGE, false, 7)
 NNE_DML_REGISTER_POOLING_OP(AveragePool, AVERAGE, false, 10)
 NNE_DML_REGISTER_POOLING_OP(AveragePool, AVERAGE, false, 11)
+NNE_DML_REGISTER_POOLING_OP(AveragePool, AVERAGE, false, 19)
 NNE_DML_REGISTER_POOLING_OP(GlobalAveragePool, AVERAGE, true, 1)
 NNE_DML_REGISTER_POOLING_OP(LpPool, LP, false, 2)
 NNE_DML_REGISTER_POOLING_OP(LpPool, LP, false, 11)
diff --git a/Engine/Plugins/Experimental/NNERuntimeRDG/Source/NNERuntimeRDG/Private/Hlsl/NNERuntimeRDGPool.cpp b/Engine/Plugins/Experimental/NNERuntimeRDG/Source/NNERuntimeRDG/Private/Hlsl/NNERuntimeRDGPool.cpp
--- a/Engine/Plugins/Experimental/NNERuntimeRDG/Source/NNERuntimeRDG/Private/Hlsl/NNERuntimeRDGPool.cpp
+++ b/Engine/Plugins/Experimental/NNERuntimeRDG/Source/NNERuntimeRDG/Private/Hlsl/NNERuntimeRDGPool.cpp
@@ -241,9 +241,6 @@ namespace UE::NNERuntimeRDG::Private::Hlsl
 	template< UE::NNEHlslShaders::Internal::EPoolOperatorType PoolOperatorType >
 	bool ValidatePoolOperator(const NNE::FAttributeMap& AttributeMap, TConstArrayView<ENNETensorDataType> InputTypes, TConstArrayView<NNE::FSymbolicTensorShape> InputPools)
 	{
-		//This match version 8 of the MaxPool operator, next version is 10
-		//and version 7 of the AveragexPool operator, next version is 10
-		//[REDACTED_PROFILE_URL]/onnx/blob/main/docs/Operators.md#MaxPool
 		bool bIsValid = true;
 
 		FAttributeValidator AttributeValidator;
@@ -282,8 +279,9 @@ namespace UE::NNERuntimeRDG::Private::Hlsl
 
 	bool RegisterPoolOperators(FOperatorRegistryHlsl& Registry)
 	{
-		Registry.OpAdd({{TEXT("MaxPool"), TEXT("Onnx")}}, CreateMaxPoolOperator, ValidatePoolOperator<UE::NNEHlslShaders::Internal::EPoolOperatorType::MAX_POOL>);
-		Registry.OpAdd({ {TEXT("AveragePool"), TEXT("Onnx")} }, CreateAveragePoolOperator, ValidatePoolOperator<UE::NNEHlslShaders::Internal::EPoolOperatorType::AVERAGE_POOL>);
+		// Note: support of a particular version is partial with respect to tensor data types (only the most typical ones are usually supported).
+		Registry.OpAdd({{TEXT("MaxPool"), TEXT("Onnx")}, 8}, CreateMaxPoolOperator, ValidatePoolOperator<UE::NNEHlslShaders::Internal::EPoolOperatorType::MAX_POOL>);
+		Registry.OpAdd({ {TEXT("AveragePool"), TEXT("Onnx")}, 7}, CreateAveragePoolOperator, ValidatePoolOperator<UE::NNEHlslShaders::Internal::EPoolOperatorType::AVERAGE_POOL>);
 		return true;
 	}
 } // UE::NNERuntimeRDG::Private::Hlsl
diff --git a/Engine/Plugins/NNE/NNERuntimeORT/Source/NNEUtilities/Internal/NNEUtilitiesHelpers.h b/Engine/Plugins/NNE/NNERuntimeORT/Source/NNEUtilities/Internal/NNEUtilitiesHelpers.h
new file mode 100644
--- /dev/null
+++ b/Engine/Plugins/NNE/NNERuntimeORT/Source/NNEUtilities/Internal/NNEUtilitiesHelpers.h
@@ -0,0 +1,11 @@
+// Copyright Epic Games, Inc. All Rights Reserved.
+#pragma once
+
+#include "CoreMinimal.h"
+
+namespace UE::NNEUtilities::Internal
+{
+
+NNEUTILITIES_API TOptional<uint32> GetOpVersionFromOpsetVersion(const FString& OpType, int OpsetVersion);
+
+} // namespace UE::NNEUtilities::Internal
\ No newline at end of file
diff --git a/Engine/Plugins/NNE/NNERuntimeORT/Source/NNEUtilities/Internal/NNEUtilitiesModelBuilderONNX.h b/Engine/Plugins/NNE/NNERuntimeORT/Source/NNEUtilities/Internal/NNEUtilitiesModelBuilderONNX.h
--- a/Engine/Plugins/NNE/NNERuntimeORT/Source/NNEUtilities/Internal/NNEUtilitiesModelBuilderONNX.h
+++ b/Engine/Plugins/NNE/NNERuntimeORT/Source/NNEUtilities/Internal/NNEUtilitiesModelBuilderONNX.h
@@ -15,12 +15,9 @@ struct FNNEModelRaw;
 namespace UE::NNEUtilities::Internal
 {
 
-static constexpr int64 DefaultOnnxIrVersion = 7;
-static constexpr int64 DefaultOnnxOpsetVersion = 15;
+NNEUTILITIES_API TUniquePtr<IModelBuilder> CreateONNXModelBuilder(int64 IrVersion, int64 OpsetVersion);
 
-NNEUTILITIES_API TUniquePtr<IModelBuilder> CreateONNXModelBuilder(int64 IrVersion = DefaultOnnxIrVersion, int64 OpsetVersion = DefaultOnnxOpsetVersion);
-
-NNEUTILITIES_API bool CreateONNXModelForOperator(const FString& OperatorName, bool bUseVariadicShapeForModel,
+NNEUTILITIES_API bool CreateONNXModelForOperator(const FString& OperatorName, int32 IrVersion, int32 OpsetVersion, bool bUseVariadicShapeForModel,
 	TConstArrayView<NNE::Internal::FTensor> InInputTensors, TConstArrayView<NNE::Internal::FTensor> InOutputTensors,
 	TConstArrayView<NNE::Internal::FTensor> InWeightTensors, TConstArrayView<TConstArrayView<uint8>> InWeightTensorsData,
 	const UE::NNE::FAttributeMap& Attributes, FNNEModelRaw& ModelData);
diff --git a/Engine/Plugins/NNE/NNERuntimeORT/Source/NNEUtilities/Private/NNEUtilitiesHelpers.cpp b/Engine/Plugins/NNE/NNERuntimeORT/Source/NNEUtilities/Private/NNEUtilitiesHelpers.cpp
new file mode 100644
--- /dev/null
+++ b/Engine/Plugins/NNE/NNERuntimeORT/Source/NNEUtilities/Private/NNEUtilitiesHelpers.cpp
@@ -0,0 +1,30 @@
+// Copyright Epic Games, Inc. All Rights Reserved.
+
+#include "NNEUtilitiesHelpers.h"
+
+#include "NNE.h"
+
+#include "NNEUtilitiesThirdPartyWarningDisabler.h"
+NNE_THIRD_PARTY_INCLUDES_START
+#undef check
+#undef TEXT
+
+#include <onnx/defs/schema.h>
+NNE_THIRD_PARTY_INCLUDES_END
+
+
+namespace UE::NNEUtilities::Internal
+{
+
+TOptional<uint32> GetOpVersionFromOpsetVersion(const FString& OpType, int OpsetVersion)
+{
+	const onnx::OpSchema* OpSchema = onnx::OpSchemaRegistry::Schema(TCHAR_TO_ANSI(*OpType), OpsetVersion);
+	if(OpSchema == nullptr)
+	{
+		UE_LOG(LogNNE, Warning, TEXT("No OpSchema found for operator %s and OpSet version %d."), *OpType, OpsetVersion);
+		return TOptional<uint32>();
+	}
+	return (uint32) OpSchema->SinceVersion();
+}
+
+}
\ No newline at end of file
diff --git a/Engine/Plugins/NNE/NNERuntimeORT/Source/NNEUtilities/Private/NNEUtilitiesModelBuilderONNX.cpp b/Engine/Plugins/NNE/NNERuntimeORT/Source/NNEUtilities/Private/NNEUtilitiesModelBuilderONNX.cpp
--- a/Engine/Plugins/NNE/NNERuntimeORT/Source/NNEUtilities/Private/NNEUtilitiesModelBuilderONNX.cpp
+++ b/Engine/Plugins/NNE/NNERuntimeORT/Source/NNEUtilities/Private/NNEUtilitiesModelBuilderONNX.cpp
@@ -64,12 +64,12 @@ class FModelBuilderONNX : public IModelBuilder
 
 	onnx::ModelProto	Model;
 	onnx::GraphProto*	Graph{ nullptr };
-	int64				IrVersion { DefaultOnnxIrVersion };
-	int64				OpsetVersion{ DefaultOnnxOpsetVersion };
+	int64				IrVersion;
+	int64				OpsetVersion;
 	
 public:
 
-	FModelBuilderONNX(int64 InIrVersion = DefaultOnnxIrVersion, int64 InOpsetVersion = DefaultOnnxOpsetVersion)
+	FModelBuilderONNX(int64 InIrVersion, int64 InOpsetVersion)
 		: IrVersion(InIrVersion) , OpsetVersion(InOpsetVersion) {}
 	
 	virtual bool Begin(const FString& Name) override
@@ -321,49 +321,14 @@ private:
 	}
 };
 
-bool CreateONNXModelForOperator(const FString& OperatorName, bool bUseVariadicShapeForModel,
+bool CreateONNXModelForOperator(const FString& OperatorName, int32 IrVersion, int32 OpsetVersion, bool bUseVariadicShapeForModel,
 	TConstArrayView<NNE::Internal::FTensor> InInputTensors, TConstArrayView<NNE::Internal::FTensor> InOutputTensors,
 	TConstArrayView<NNE::Internal::FTensor> InWeightTensors, TConstArrayView<TConstArrayView<uint8>> InWeightTensorsData,
 	const UE::NNE::FAttributeMap& Attributes, FNNEModelRaw& Model)
 {
 	Model = FNNEModelRaw{};
-	
-	int64 IrVersion = DefaultOnnxIrVersion;
-	int64 OpsetVersion = DefaultOnnxOpsetVersion;
-
-	if (OperatorName == TEXT("BatchNormalization") ||	// current implementation is opset 9 (next version is 14)
-		OperatorName == TEXT("Clip") ||					// current implementation is opset 6 (next version is 11)
-		OperatorName == TEXT("Pad") ||					// current implementation is opset 2 (next version is 11)
-		OperatorName == TEXT("Split") ||				// current implementation is opset 2 (next version is 11)
-		OperatorName == TEXT("Shape") ||				// current implementation is opset 1 (next version is 13)
-		OperatorName == TEXT("Slice") ||				// current implementation is opset 1 (next version is 10)
-		OperatorName == TEXT("Squeeze") ||				// current implementation is opset 1 (next version is 11)
-		OperatorName == TEXT("Unsqueeze") ||			// current implementation is opset 1 (next version is 11)
-		OperatorName == TEXT("Upsample")				// deprecated starting opset 10
-		)				
-	{
-		OpsetVersion = 9;
-	}
-	else
-	if (OperatorName == TEXT("ReduceL1") ||				// current implementation is opset 11 (next version is 13)
-		OperatorName == TEXT("ReduceL2") ||				// current implementation is opset 11 (next version is 13)
-		OperatorName == TEXT("ReduceLogSum") ||			// current implementation is opset 11 (next version is 13)
-		OperatorName == TEXT("ReduceLogSumExp") ||		// current implementation is opset 11 (next version is 13)
-		OperatorName == TEXT("ReduceLogMin") ||			// current implementation is opset 11 (next version is 13)
-		OperatorName == TEXT("ReduceLogMax") ||			// current implementation is opset 11 (next version is 13)
-		OperatorName == TEXT("ReduceMean") ||			// current implementation is opset 11 (next version is 13)
-		OperatorName == TEXT("ReduceProd") ||			// current implementation is opset 11 (next version is 13)
-		OperatorName == TEXT("ReduceSum") ||			// current implementation is opset 11 (next version is 13)
-		OperatorName == TEXT("ReduceSumSquare") ||		// current implementation is opset 11 (next version is 13)
-		OperatorName == TEXT("Resize") ||				// current implementation is opset 11 (next version is 13)
-		OperatorName == TEXT("Squeeze") ||				// current implementation is opset 11 (next version is 13)
-		OperatorName == TEXT("Unsqueeze")				// current implementation is opset 11 (next version is 13)
-		)
-	{
-		OpsetVersion = 11;
-	}
 
-	TUniquePtr<IModelBuilder> Builder(CreateONNXModelBuilder(IrVersion, OpsetVersion));
+	TUniquePtr<IModelBuilder> Builder(CreateONNXModelBuilder((int64) IrVersion, (int64) OpsetVersion));
 
 	Builder->Begin();
 
diff --git a/Engine/Plugins/NNE/NNERuntimeORT/Source/NNEUtilities/Private/NNEUtilitiesModelOptimizerNNE.cpp b/Engine/Plugins/NNE/NNERuntimeORT/Source/NNEUtilities/Private/NNEUtilitiesModelOptimizerNNE.cpp
--- a/Engine/Plugins/NNE/NNERuntimeORT/Source/NNEUtilities/Private/NNEUtilitiesModelOptimizerNNE.cpp
+++ b/Engine/Plugins/NNE/NNERuntimeORT/Source/NNEUtilities/Private/NNEUtilitiesModelOptimizerNNE.cpp
@@ -6,6 +6,7 @@
 #include "NNERuntimeFormat.h"
 #include "NNEUtilitiesModelBuilderNNE.h"
 #include "NNEUtilitiesModelOptimizerONNX.h"
+#include "NNEUtilitiesHelpers.h"
 
 #include "NNEUtilitiesThirdPartyWarningDisabler.h"
 NNE_THIRD_PARTY_INCLUDES_START
@@ -182,17 +183,6 @@ namespace ModelOptimizerNNEHelper
 		}
 	}
 
-	TOptional<uint32> GetOpVersion(const onnx::ModelProto& ModelProto, const FString& OpType)
-	{
-		const onnx::OpSchema* OpSchema = onnx::OpSchemaRegistry::Schema(TCHAR_TO_ANSI(*OpType), (int) ModelProto.opset_import(0).version());
-		if(OpSchema == nullptr)
-		{
-			UE_LOG(LogNNE, Warning, TEXT("No OpSchema found for operator %s and OpSet version %d."), *OpType, (int) ModelProto.opset_import(0).version());
-			return TOptional<uint32>();
-		}
-		return (uint32) OpSchema->SinceVersion();
-	}
-
 	bool BuildNNEFormatFromONNX(TArray<uint8>& ONNXData, TArray<uint8>& NNEData)
 	{
 		TUniquePtr<IModelBuilder> Builder = CreateNNEModelBuilder();
@@ -277,7 +267,7 @@ namespace ModelOptimizerNNEHelper
 					Builder->AddTensor(FString(ANSI_TO_TCHAR(Output.name().c_str())) + TEXT("_NNEInitializer"), DataType, Shape, Data, DataSize);
 
 				const FString IdentityOpType = TEXT("Identity");
-				TOptional<uint32> OpVersion = GetOpVersion(ModelProto, IdentityOpType);
+				TOptional<uint32> OpVersion = GetOpVersionFromOpsetVersion(IdentityOpType, (int) ModelProto.opset_import(0).version());
 				if(!OpVersion.IsSet())
 				{
 					return false;
@@ -296,7 +286,7 @@ namespace ModelOptimizerNNEHelper
 			FString NNEOpType(StringCast<TCHAR>(OnnxOpType.c_str()));
 			const std::string& OnnxOpName = Node.name();
 			FString NNEOpName(StringCast<TCHAR>(OnnxOpName.c_str()));
-			TOptional<uint32> OpVersion = GetOpVersion(ModelProto, NNEOpType);
+			TOptional<uint32> OpVersion = GetOpVersionFromOpsetVersion(NNEOpType, (int) ModelProto.opset_import(0).version());
 			if(!OpVersion.IsSet())
 			{
 				return false;
