[NNE] Error status use Enum instead of int.

Problem:
- Runtime interface member functions SetInputTensorShapes, RunSync and EnqueueRDG return int as a status code
- Failue code is not determined and might conflict between implementations

Solution:
- Replace status code int with enum type to add semantics
- Can be extended internally or externally

Tests:
- Compile Files (Non-Unity)
- Unreal Editor, Win64: Style-transfer, smoke tests
- Unreal Game, Win64, Style-transfer, smoke tests

#rnx

#tests Non-Unity compilation, Editor/Game (Win64) style-transfer and smoke tests

diff --git a/Engine/Plugins/Animation/MLDeformer/VertexDeltaModel/Source/VertexDeltaModel/Private/VertexDeltaModelInstance.cpp b/Engine/Plugins/Animation/MLDeformer/VertexDeltaModel/Source/VertexDeltaModel/Private/VertexDeltaModelInstance.cpp
--- a/Engine/Plugins/Animation/MLDeformer/VertexDeltaModel/Source/VertexDeltaModel/Private/VertexDeltaModelInstance.cpp
+++ b/Engine/Plugins/Animation/MLDeformer/VertexDeltaModel/Source/VertexDeltaModel/Private/VertexDeltaModelInstance.cpp
@@ -85,7 +85,7 @@ void UVertexDeltaModelInstance::Execute(float ModelWeight)
 					{
 						return;
 					}
-					int32 Res = ModelInstanceRDG->EnqueueRDG(GraphBuilder, InputBindingsRDG, OutputBindingsRDG);
+					ModelInstanceRDG->EnqueueRDG(GraphBuilder, InputBindingsRDG, OutputBindingsRDG);
 	
 					GraphBuilder.Execute();
 				}
diff --git a/Engine/Plugins/Experimental/NNEDenoiser/Source/NNEDenoiser/Private/NNEDenoiserGenericDenoiser.cpp b/Engine/Plugins/Experimental/NNEDenoiser/Source/NNEDenoiser/Private/NNEDenoiserGenericDenoiser.cpp
--- a/Engine/Plugins/Experimental/NNEDenoiser/Source/NNEDenoiser/Private/NNEDenoiserGenericDenoiser.cpp
+++ b/Engine/Plugins/Experimental/NNEDenoiser/Source/NNEDenoiser/Private/NNEDenoiserGenericDenoiser.cpp
@@ -94,8 +94,8 @@ void AddTilePasses(FRDGBuilder& GraphBuilder, IModelInstance& ModelInstance, con
 	InputProcess.AddPasses(GraphBuilder, ModelInstance.GetInputTensorDescs(), ModelInstance.GetInputTensorShapes(), ResourceAccess, InputBuffers);
 
 	// 2. Create buffer binding and infer the model
-	bool bSuccess = !ModelInstance.EnqueueRDG(GraphBuilder, GetBindingRDG(InputBuffers), GetBindingRDG(OutputBuffers));
-	check(bSuccess);
+	NNE::IModelInstanceRDG::EEnqueueRDGStatus Status = ModelInstance.EnqueueRDG(GraphBuilder, GetBindingRDG(InputBuffers), GetBindingRDG(OutputBuffers));
+	checkf(Status == NNE::IModelInstanceRDG::EEnqueueRDGStatus::Ok, TEXT("EnqueueRDG failed: %d"), static_cast<int>(Status));
 
 	// 3. Write output based on output buffer
 	FRDGTextureRef OutputTexture = ResourceManager.GetTexture(EResourceName::Output, 0);
diff --git a/Engine/Plugins/Experimental/NNEDenoiser/Source/NNEDenoiser/Private/NNEDenoiserIOProcessBase.cpp b/Engine/Plugins/Experimental/NNEDenoiser/Source/NNEDenoiser/Private/NNEDenoiserIOProcessBase.cpp
--- a/Engine/Plugins/Experimental/NNEDenoiser/Source/NNEDenoiser/Private/NNEDenoiserIOProcessBase.cpp
+++ b/Engine/Plugins/Experimental/NNEDenoiser/Source/NNEDenoiser/Private/NNEDenoiserIOProcessBase.cpp
@@ -189,8 +189,8 @@ bool FInputProcessBase::PrepareAndValidate(IModelInstance& ModelInstance, FIntPo
 		UE_LOG(LogNNEDenoiser, Log, TEXT("%d: (%d, %d, %d, %d)"), Idx, ModelInputShape.X, ModelInputShape.Y, ModelInputShape.Z, ModelInputShape.W)
 	}
 
-	const int32 Success = ModelInstance.SetInputTensorShapes(InputShapes);
-	if (Success != 0)
+	const IModelInstance::ESetInputTensorShapeStatus Status = ModelInstance.SetInputTensorShapes(InputShapes);
+	if (Status != IModelInstance::ESetInputTensorShapeStatus::Ok)
 	{
 		UE_LOG(LogNNEDenoiser, Error, TEXT("Could not configure model instance (ModelInstance.SetInputTensorShapes() failed)!"))
 		return false;
diff --git a/Engine/Plugins/Experimental/NNEDenoiser/Source/NNEDenoiser/Private/NNEDenoiserModelInstance.h b/Engine/Plugins/Experimental/NNEDenoiser/Source/NNEDenoiser/Private/NNEDenoiserModelInstance.h
--- a/Engine/Plugins/Experimental/NNEDenoiser/Source/NNEDenoiser/Private/NNEDenoiserModelInstance.h
+++ b/Engine/Plugins/Experimental/NNEDenoiser/Source/NNEDenoiser/Private/NNEDenoiserModelInstance.h
@@ -10,6 +10,9 @@ namespace UE::NNEDenoiser::Private
 class IModelInstance : public NNE::IModelInstanceRDG
 {
 public:
+	using ESetInputTensorShapeStatus = NNE::IModelInstanceRDG::ESetInputTensorShapeStatus;
+	using EEnqueueRDGStatus = NNE::IModelInstanceRDG::EEnqueueRDGStatus;
+
 	virtual ~IModelInstance() = default;
 };
 
diff --git a/Engine/Plugins/Experimental/NNEDenoiser/Source/NNEDenoiser/Private/NNEDenoiserModelInstanceCPU.cpp b/Engine/Plugins/Experimental/NNEDenoiser/Source/NNEDenoiser/Private/NNEDenoiserModelInstanceCPU.cpp
--- a/Engine/Plugins/Experimental/NNEDenoiser/Source/NNEDenoiser/Private/NNEDenoiserModelInstanceCPU.cpp
+++ b/Engine/Plugins/Experimental/NNEDenoiser/Source/NNEDenoiser/Private/NNEDenoiserModelInstanceCPU.cpp
@@ -91,12 +91,14 @@ namespace UE::NNEDenoiser::Private
 		return ModelInstance->GetOutputTensorShapes();
 	}
 
-	int32 FModelInstanceCPU::SetInputTensorShapes(TConstArrayView<NNE::FTensorShape> InInputShapes)
+	FModelInstanceCPU::ESetInputTensorShapeStatus FModelInstanceCPU::SetInputTensorShapes(TConstArrayView<NNE::FTensorShape> InInputShapes)
 	{
-		return ModelInstance->SetInputTensorShapes(InInputShapes);
+		NNE::IModelInstanceCPU::ESetInputTensorShapeStatus Status = ModelInstance->SetInputTensorShapes(InInputShapes);
+
+		return Status == NNE::IModelInstanceCPU::ESetInputTensorShapeStatus::Ok ? ESetInputTensorShapeStatus::Ok : ESetInputTensorShapeStatus::Fail;
 	}
 
-	int32 FModelInstanceCPU::EnqueueRDG(FRDGBuilder &GraphBuilder, TConstArrayView<NNE::FTensorBindingRDG> Inputs, TConstArrayView<NNE::FTensorBindingRDG> Outputs)
+	FModelInstanceCPU::EEnqueueRDGStatus FModelInstanceCPU::EnqueueRDG(FRDGBuilder &GraphBuilder, TConstArrayView<NNE::FTensorBindingRDG> Inputs, TConstArrayView<NNE::FTensorBindingRDG> Outputs)
 	{
 		FNNEDenoiserModelInstanceCPUTextureParameters* DenoiserParameters = GraphBuilder.AllocParameters<FNNEDenoiserModelInstanceCPUTextureParameters>();
 		for (const NNE::FTensorBindingRDG& Binding : Inputs)
@@ -143,7 +145,8 @@ namespace UE::NNEDenoiser::Private
 				CopyBufferFromGPUToCPU(RHICmdList, Buffer, Buffer->GetSize(), ScratchInputBuffers[Idx]);
 			}
 
-			ModelInstance->RunSync(InputBindings, OutputBindings);
+			NNE::IModelInstanceCPU::ERunSyncStatus Status = ModelInstance->RunSync(InputBindings, OutputBindings);
+			checkf(Status == NNE::IModelInstanceCPU::ERunSyncStatus::Ok, TEXT("RunSync failed with status %d"), static_cast<int>(Status))
 
 			for (int32 Idx = 0; Idx < DenoiserParameters->OutputBuffers.Num(); Idx++)
 			{
@@ -158,7 +161,7 @@ namespace UE::NNEDenoiser::Private
 #endif
 		});
 
-		return 0;
+		return EEnqueueRDGStatus::Ok;
 	}
 
 }
\ No newline at end of file
diff --git a/Engine/Plugins/Experimental/NNEDenoiser/Source/NNEDenoiser/Private/NNEDenoiserModelInstanceCPU.h b/Engine/Plugins/Experimental/NNEDenoiser/Source/NNEDenoiser/Private/NNEDenoiserModelInstanceCPU.h
--- a/Engine/Plugins/Experimental/NNEDenoiser/Source/NNEDenoiser/Private/NNEDenoiserModelInstanceCPU.h
+++ b/Engine/Plugins/Experimental/NNEDenoiser/Source/NNEDenoiser/Private/NNEDenoiserModelInstanceCPU.h
@@ -32,9 +32,9 @@ public:
 
 	virtual TConstArrayView<NNE::FTensorShape> GetOutputTensorShapes() const override;
 
-	virtual int32 SetInputTensorShapes(TConstArrayView<NNE::FTensorShape> InInputShapes) override;
+	virtual ESetInputTensorShapeStatus SetInputTensorShapes(TConstArrayView<NNE::FTensorShape> InInputShapes) override;
 
-	virtual int32 EnqueueRDG(FRDGBuilder &GraphBuilder, TConstArrayView<NNE::FTensorBindingRDG> Inputs, TConstArrayView<NNE::FTensorBindingRDG> Outputs) override;
+	virtual EEnqueueRDGStatus EnqueueRDG(FRDGBuilder &GraphBuilder, TConstArrayView<NNE::FTensorBindingRDG> Inputs, TConstArrayView<NNE::FTensorBindingRDG> Outputs) override;
 
 private:
 	TSharedRef<UE::NNE::IModelInstanceCPU> ModelInstance;
diff --git a/Engine/Plugins/Experimental/NNEDenoiser/Source/NNEDenoiser/Private/NNEDenoiserModelInstanceGPU.cpp b/Engine/Plugins/Experimental/NNEDenoiser/Source/NNEDenoiser/Private/NNEDenoiserModelInstanceGPU.cpp
--- a/Engine/Plugins/Experimental/NNEDenoiser/Source/NNEDenoiser/Private/NNEDenoiserModelInstanceGPU.cpp
+++ b/Engine/Plugins/Experimental/NNEDenoiser/Source/NNEDenoiser/Private/NNEDenoiserModelInstanceGPU.cpp
@@ -94,12 +94,14 @@ namespace UE::NNEDenoiser::Private
 		return ModelInstance->GetOutputTensorShapes();
 	}
 
-	int32 FModelInstanceGPU::SetInputTensorShapes(TConstArrayView<NNE::FTensorShape> InInputShapes)
+	FModelInstanceGPU::ESetInputTensorShapeStatus FModelInstanceGPU::SetInputTensorShapes(TConstArrayView<NNE::FTensorShape> InInputShapes)
 	{
-		return ModelInstance->SetInputTensorShapes(InInputShapes);
+		NNE::IModelInstanceGPU::ESetInputTensorShapeStatus Status = ModelInstance->SetInputTensorShapes(InInputShapes);
+
+		return Status == NNE::IModelInstanceGPU::ESetInputTensorShapeStatus::Ok ? ESetInputTensorShapeStatus::Ok : ESetInputTensorShapeStatus::Fail;
 	}
 
-	int32 FModelInstanceGPU::EnqueueRDG(FRDGBuilder &GraphBuilder, TConstArrayView<NNE::FTensorBindingRDG> Inputs, TConstArrayView<NNE::FTensorBindingRDG> Outputs)
+	FModelInstanceGPU::EEnqueueRDGStatus FModelInstanceGPU::EnqueueRDG(FRDGBuilder &GraphBuilder, TConstArrayView<NNE::FTensorBindingRDG> Inputs, TConstArrayView<NNE::FTensorBindingRDG> Outputs)
 	{
 		FNNEDenoiserModelInstanceGPUTextureParameters* DenoiserParameters = GraphBuilder.AllocParameters<FNNEDenoiserModelInstanceGPUTextureParameters>();
 		for (const NNE::FTensorBindingRDG& Binding : Inputs)
@@ -149,7 +151,8 @@ namespace UE::NNEDenoiser::Private
 				CopyBufferFromGPUToCPU(RHICmdList, Buffer, Buffer->GetSize(), ScratchInputBuffers[Idx]);
 			}
 
-			ModelInstance->RunSync(InputBindings, OutputBindings);
+			NNE::IModelInstanceGPU::ERunSyncStatus Status = ModelInstance->RunSync(InputBindings, OutputBindings);
+			checkf(Status == NNE::IModelInstanceGPU::ERunSyncStatus::Ok, TEXT("RunSync failed with status %d"), static_cast<int32>(Status))
 
 			for (int32 Idx = 0; Idx < DenoiserParameters->OutputBuffers.Num(); Idx++)
 			{
@@ -164,7 +167,7 @@ namespace UE::NNEDenoiser::Private
 #endif
 		});
 
-		return 0;
+		return EEnqueueRDGStatus::Ok;
 	}
 
 }
\ No newline at end of file
diff --git a/Engine/Plugins/Experimental/NNEDenoiser/Source/NNEDenoiser/Private/NNEDenoiserModelInstanceGPU.h b/Engine/Plugins/Experimental/NNEDenoiser/Source/NNEDenoiser/Private/NNEDenoiserModelInstanceGPU.h
--- a/Engine/Plugins/Experimental/NNEDenoiser/Source/NNEDenoiser/Private/NNEDenoiserModelInstanceGPU.h
+++ b/Engine/Plugins/Experimental/NNEDenoiser/Source/NNEDenoiser/Private/NNEDenoiserModelInstanceGPU.h
@@ -30,9 +30,9 @@ public:
 
 	virtual TConstArrayView<NNE::FTensorShape> GetOutputTensorShapes() const override;
 
-	virtual int32 SetInputTensorShapes(TConstArrayView<NNE::FTensorShape> InInputShapes) override;
+	virtual ESetInputTensorShapeStatus SetInputTensorShapes(TConstArrayView<NNE::FTensorShape> InInputShapes) override;
 
-	virtual int32 EnqueueRDG(FRDGBuilder &GraphBuilder, TConstArrayView<NNE::FTensorBindingRDG> Inputs, TConstArrayView<NNE::FTensorBindingRDG> Outputs) override;
+	virtual EEnqueueRDGStatus EnqueueRDG(FRDGBuilder &GraphBuilder, TConstArrayView<NNE::FTensorBindingRDG> Inputs, TConstArrayView<NNE::FTensorBindingRDG> Outputs) override;
 
 private:
 	TSharedRef<UE::NNE::IModelInstanceGPU> ModelInstance;
diff --git a/Engine/Plugins/Experimental/NNEDenoiser/Source/NNEDenoiser/Private/NNEDenoiserModelInstanceRDG.cpp b/Engine/Plugins/Experimental/NNEDenoiser/Source/NNEDenoiser/Private/NNEDenoiserModelInstanceRDG.cpp
--- a/Engine/Plugins/Experimental/NNEDenoiser/Source/NNEDenoiser/Private/NNEDenoiserModelInstanceRDG.cpp
+++ b/Engine/Plugins/Experimental/NNEDenoiser/Source/NNEDenoiser/Private/NNEDenoiserModelInstanceRDG.cpp
@@ -87,12 +87,12 @@ namespace UE::NNEDenoiser::Private
 		return ModelInstance->GetOutputTensorShapes();
 	}
 
-	int32 FModelInstanceRDG::SetInputTensorShapes(TConstArrayView<NNE::FTensorShape> InInputShapes)
+	FModelInstanceRDG::ESetInputTensorShapeStatus FModelInstanceRDG::SetInputTensorShapes(TConstArrayView<NNE::FTensorShape> InInputShapes)
 	{
 		return ModelInstance->SetInputTensorShapes(InInputShapes);
 	}
 
-	int32 FModelInstanceRDG::EnqueueRDG(FRDGBuilder &GraphBuilder, TConstArrayView<NNE::FTensorBindingRDG> Inputs, TConstArrayView<NNE::FTensorBindingRDG> Outputs)
+	FModelInstanceRDG::EEnqueueRDGStatus FModelInstanceRDG::EnqueueRDG(FRDGBuilder &GraphBuilder, TConstArrayView<NNE::FTensorBindingRDG> Inputs, TConstArrayView<NNE::FTensorBindingRDG> Outputs)
 	{
 		return ModelInstance->EnqueueRDG(GraphBuilder, Inputs, Outputs);
 	}
diff --git a/Engine/Plugins/Experimental/NNEDenoiser/Source/NNEDenoiser/Private/NNEDenoiserModelInstanceRDG.h b/Engine/Plugins/Experimental/NNEDenoiser/Source/NNEDenoiser/Private/NNEDenoiserModelInstanceRDG.h
--- a/Engine/Plugins/Experimental/NNEDenoiser/Source/NNEDenoiser/Private/NNEDenoiserModelInstanceRDG.h
+++ b/Engine/Plugins/Experimental/NNEDenoiser/Source/NNEDenoiser/Private/NNEDenoiserModelInstanceRDG.h
@@ -30,9 +30,9 @@ public:
 
 	virtual TConstArrayView<NNE::FTensorShape> GetOutputTensorShapes() const override;
 
-	virtual int32 SetInputTensorShapes(TConstArrayView<NNE::FTensorShape> InInputShapes) override;
+	virtual ESetInputTensorShapeStatus SetInputTensorShapes(TConstArrayView<NNE::FTensorShape> InInputShapes) override;
 
-	virtual int32 EnqueueRDG(FRDGBuilder &GraphBuilder, TConstArrayView<NNE::FTensorBindingRDG> Inputs, TConstArrayView<NNE::FTensorBindingRDG> Outputs) override;
+	virtual EEnqueueRDGStatus EnqueueRDG(FRDGBuilder &GraphBuilder, TConstArrayView<NNE::FTensorBindingRDG> Inputs, TConstArrayView<NNE::FTensorBindingRDG> Outputs) override;
 
 private:
 	TSharedRef<UE::NNE::IModelInstanceRDG> ModelInstance;
diff --git a/Engine/Plugins/Experimental/NNERuntimeBasicCpu/Source/NNERuntimeBasicCpu/Private/NNERuntimeBasicCpuModel.cpp b/Engine/Plugins/Experimental/NNERuntimeBasicCpu/Source/NNERuntimeBasicCpu/Private/NNERuntimeBasicCpuModel.cpp
--- a/Engine/Plugins/Experimental/NNERuntimeBasicCpu/Source/NNERuntimeBasicCpu/Private/NNERuntimeBasicCpuModel.cpp
+++ b/Engine/Plugins/Experimental/NNERuntimeBasicCpu/Source/NNERuntimeBasicCpu/Private/NNERuntimeBasicCpuModel.cpp
@@ -3383,20 +3383,20 @@ namespace UE::NNE::RuntimeBasic
 		, Instance(Model->Layer->MakeInstance())
 	{}
 
-	int FModelInstanceCPU::SetInputTensorShapes(TConstArrayView<FTensorShape> InInputShapes)
+	FModelInstanceCPU::ESetInputTensorShapeStatus FModelInstanceCPU::SetInputTensorShapes(TConstArrayView<FTensorShape> InInputShapes)
 	{
 		NNE_RUNTIME_BASIC_TRACE_SCOPE(NNE::RuntimeBasic::FModelInstanceCPU::SetInputTensorShapes);
 
 		if (!ensureMsgf(InInputShapes.Num() == 1, TEXT("Basic CPU Inference only supports single input tensor.")))
 		{
-			return -1;
+			return ESetInputTensorShapeStatus::Fail;
 		}
 
 		const FTensorShape& InputShape = InInputShapes[0];
 
 		if (!ensureMsgf(InputShape.Rank() == 2, TEXT("Basic CPU Inference only supports rank 2 input tensors.")))
 		{
-			return -1;
+			return ESetInputTensorShapeStatus::Fail;
 		}
 
 		const uint32 InputBatchSize = InputShape.GetData()[0];
@@ -3406,7 +3406,7 @@ namespace UE::NNE::RuntimeBasic
 
 		if (!ensureMsgf(InputInputSize == ModelInputSize, TEXT("Input tensor shape does not match model input size. Got %i, expected %i."), InputInputSize, ModelInputSize))
 		{
-			return -1;
+			return ESetInputTensorShapeStatus::Fail;
 		}
 
 		BatchSize = InputBatchSize;
@@ -3421,36 +3421,36 @@ namespace UE::NNE::RuntimeBasic
 			Instance->SetMaxBatchSize(BatchSize);
 		}
 
-		return 0;
+		return ESetInputTensorShapeStatus::Ok;
 	}
 
-	int FModelInstanceCPU::RunSync(TConstArrayView<FTensorBindingCPU> InInputBindings, TConstArrayView<FTensorBindingCPU> InOutputBindings)
+	FModelInstanceCPU::ERunSyncStatus FModelInstanceCPU::RunSync(TConstArrayView<FTensorBindingCPU> InInputBindings, TConstArrayView<FTensorBindingCPU> InOutputBindings)
 	{
 		NNE_RUNTIME_BASIC_TRACE_SCOPE(NNE::RuntimeBasic::FModelInstanceCPU::RunSync);
 
 		if (!ensureMsgf(BatchSize > 0, TEXT("SetInputTensorShapes must be run before RunSync")))
 		{
-			return -1;
+			return ERunSyncStatus::Fail;
 		}
 
 		if (!ensureMsgf(InInputBindings.Num() == 1, TEXT("Basic CPU Inference only supports single input tensor.")))
 		{
-			return -1;
+			return ERunSyncStatus::Fail;
 		}
 
 		if (!ensureMsgf(InOutputBindings.Num() == 1, TEXT("Basic CPU Inference only supports single output tensor.")))
 		{
-			return -1;
+			return ERunSyncStatus::Fail;
 		}
 
 		if (!ensureMsgf(InInputBindings[0].SizeInBytes == BatchSize * InputSize * sizeof(float), TEXT("Incorrect Input Tensor Size")))
 		{
-			return -1;
+			return ERunSyncStatus::Fail;
 		}
 
 		if (!ensureMsgf(InOutputBindings[0].SizeInBytes == BatchSize * OutputSize * sizeof(float), TEXT("Incorrect Output Tensor Size")))
 		{
-			return -1;
+			return ERunSyncStatus::Fail;
 		}
 
 		Model->Layer->Evaluate(
@@ -3463,7 +3463,7 @@ namespace UE::NNE::RuntimeBasic
 			OutputSize,
 			InputSize);
 
-		return 0;
+		return ERunSyncStatus::Ok;
 	}
 
 	uint32 FModelCPU::ModelMagicNumber = 0x0BA51C01;
diff --git a/Engine/Plugins/Experimental/NNERuntimeBasicCpu/Source/NNERuntimeBasicCpu/Private/NNERuntimeBasicCpuModel.h b/Engine/Plugins/Experimental/NNERuntimeBasicCpu/Source/NNERuntimeBasicCpu/Private/NNERuntimeBasicCpuModel.h
--- a/Engine/Plugins/Experimental/NNERuntimeBasicCpu/Source/NNERuntimeBasicCpu/Private/NNERuntimeBasicCpuModel.h
+++ b/Engine/Plugins/Experimental/NNERuntimeBasicCpu/Source/NNERuntimeBasicCpu/Private/NNERuntimeBasicCpuModel.h
@@ -28,6 +28,8 @@ namespace UE::NNE::RuntimeBasic
 	class FModelInstanceCPU : public IModelInstanceCPU
 	{
 	public:
+		using ESetInputTensorShapeStatus = IModelInstanceCPU::ESetInputTensorShapeStatus;
+		using ERunSyncStatus = IModelInstanceCPU::ERunSyncStatus;
 
 		FModelInstanceCPU(const TSharedPtr<FModelCPU>& InModel);
 
@@ -36,8 +38,8 @@ namespace UE::NNE::RuntimeBasic
 		virtual TConstArrayView<FTensorDesc> GetOutputTensorDescs() const override final { return TConstArrayView<FTensorDesc>(&OutputTensorDesc, 1); }
 		virtual TConstArrayView<FTensorShape> GetInputTensorShapes() const override final { return TConstArrayView<FTensorShape>(&InputTensorShape, 1); }
 		virtual TConstArrayView<FTensorShape> GetOutputTensorShapes() const override final { return TConstArrayView<FTensorShape>(&OutputTensorShape, 1); }
-		virtual int SetInputTensorShapes(TConstArrayView<FTensorShape> InInputShapes) override final;
-		virtual int RunSync(TConstArrayView<FTensorBindingCPU> InInputBindings, TConstArrayView<FTensorBindingCPU> InOutputBindings) override final;
+		virtual ESetInputTensorShapeStatus SetInputTensorShapes(TConstArrayView<FTensorShape> InInputShapes) override final;
+		virtual ERunSyncStatus RunSync(TConstArrayView<FTensorBindingCPU> InInputBindings, TConstArrayView<FTensorBindingCPU> InOutputBindings) override final;
 		//~ End UE::NNE::IModelInstanceCPU Interface
 
 		TSharedPtr<FModelCPU> Model;
diff --git a/Engine/Plugins/Experimental/NNERuntimeIREE/Source/NNERuntimeIREE/Private/NNERuntimeIREEModel.cpp b/Engine/Plugins/Experimental/NNERuntimeIREE/Source/NNERuntimeIREE/Private/NNERuntimeIREEModel.cpp
--- a/Engine/Plugins/Experimental/NNERuntimeIREE/Source/NNERuntimeIREE/Private/NNERuntimeIREEModel.cpp
+++ b/Engine/Plugins/Experimental/NNERuntimeIREE/Source/NNERuntimeIREE/Private/NNERuntimeIREEModel.cpp
@@ -7,9 +7,10 @@
 #include "GenericPlatform/GenericPlatformProcess.h"
 #include "HAL/FileManager.h"
 #include "Memory/SharedBuffer.h"
-#include "NNEModelData.h"
 #include "Misc/Paths.h"
 #include "Misc/ScopeLock.h"
+#include "NNEModelData.h"
+#include "NNEStatus.h"
 #include "Serialization/Archive.h"
 
 #if PLATFORM_MICROSOFT
@@ -636,6 +637,9 @@ namespace UE::NNERuntimeIREE
 				}
 
 			public:
+				using ESetInputTensorShapeStatus = UE::NNE::IModelInstanceCPU::ESetInputTensorShapeStatus;
+				using ERunSyncStatus = UE::NNE::IModelInstanceCPU::ERunSyncStatus;
+
 				~FSession()
 				{
 					iree_runtime_call_deinitialize(&Call);
@@ -719,15 +723,15 @@ namespace UE::NNERuntimeIREE
 					return OutputTensorShapes;
 				}
 
-				int32 SetInputTensorShapes(TConstArrayView<UE::NNE::FTensorShape> InInputShapes)
+				ESetInputTensorShapeStatus SetInputTensorShapes(TConstArrayView<UE::NNE::FTensorShape> InInputShapes)
 				{
 					check(InputTensorDescs.Num() == InInputShapes.Num());
 					checkCode(for (int32 i = 0; i < InputTensorDescs.Num(); i++) { check(InInputShapes[i].IsCompatibleWith(InputTensorDescs[i].GetShape())); });
 					InputTensorShapes = InInputShapes;
-					return 0;
+					return ESetInputTensorShapeStatus::Ok;
 				}
 
-				int32 RunSyncCPU(TConstArrayView<UE::NNE::FTensorBindingCPU> InInputBindings, TConstArrayView<UE::NNE::FTensorBindingCPU> InOutputBindings)
+				ERunSyncStatus RunSyncCPU(TConstArrayView<UE::NNE::FTensorBindingCPU> InInputBindings, TConstArrayView<UE::NNE::FTensorBindingCPU> InOutputBindings)
 				{
 					check(InInputBindings.Num() == InputTensorShapes.Num());
 
@@ -767,7 +771,7 @@ namespace UE::NNERuntimeIREE
 								iree_hal_buffer_view_release(TempBufferView);
 							}
 							iree_status_free(Status);
-							return -1;
+							return ERunSyncStatus::Fail;
 						}
 
 						Status = iree_runtime_call_inputs_push_back_buffer_view(&Call, TempBufferView);
@@ -776,7 +780,7 @@ namespace UE::NNERuntimeIREE
 						{
 							UE::NNERuntimeIREE::Private::PrintIREEError("UE::NNERuntimeIREE::CPU::Private::FSession failed to push the buffer view to the input list", Status);
 							iree_status_free(Status);
-							return -1;
+							return ERunSyncStatus::Fail;
 						}
 					}
 
@@ -785,7 +789,7 @@ namespace UE::NNERuntimeIREE
 					{
 						UE::NNERuntimeIREE::Private::PrintIREEError("UE::NNERuntimeIREE::CPU::Private::FSession failed to call the model function", Status);
 						iree_status_free(Status);
-						return -1;
+						return ERunSyncStatus::Fail;
 					}
 
 					OutputTensorShapes.Reset();
@@ -819,7 +823,7 @@ namespace UE::NNERuntimeIREE
 							bCopyResults = false;
 						}
 					}
-					int32 Result = 0;
+					ERunSyncStatus Result = ERunSyncStatus::Ok;
 					if (bCopyResults)
 					{
 						for (int32 i = 0; i < InOutputBindings.Num(); i++)
@@ -828,7 +832,7 @@ namespace UE::NNERuntimeIREE
 							if (!Buffer)
 							{
 								UE_LOG(LogTemp, Error, TEXT("UE::NNERuntimeIREE::CPU::Private::FSession failed to get the result buffer"));
-								Result = -1;
+								Result = ERunSyncStatus::Fail;
 								break;
 							}
 
@@ -839,7 +843,7 @@ namespace UE::NNERuntimeIREE
 							if (!iree_status_is_ok(Status))
 							{
 								UE::NNERuntimeIREE::Private::PrintIREEError("UE::NNERuntimeIREE::CPU::Private::FSession failed to map the result buffer", Status);
-								Result = -1;
+								Result = ERunSyncStatus::Fail;
 								break;
 							}
 							FMemory::Memcpy(InOutputBindings[i].Data, BufferMapping.contents.data, DataSizeInBytes);
@@ -893,12 +897,12 @@ namespace UE::NNERuntimeIREE
 			return Session->GetOutputTensorShapes();
 		}
 
-		int32 FModelInstance::SetInputTensorShapes(TConstArrayView<UE::NNE::FTensorShape> InInputShapes)
+		FModelInstance::ESetInputTensorShapeStatus FModelInstance::SetInputTensorShapes(TConstArrayView<UE::NNE::FTensorShape> InInputShapes)
 		{
 			return Session->SetInputTensorShapes(InInputShapes);
 		}
 
-		int32 FModelInstance::RunSync(TConstArrayView<UE::NNE::FTensorBindingCPU> InInputBindings, TConstArrayView<UE::NNE::FTensorBindingCPU> InOutputBindings)
+		FModelInstance::ERunSyncStatus FModelInstance::RunSync(TConstArrayView<UE::NNE::FTensorBindingCPU> InInputBindings, TConstArrayView<UE::NNE::FTensorBindingCPU> InOutputBindings)
 		{
 			return Session->RunSyncCPU(InInputBindings, InOutputBindings);
 		}
diff --git a/Engine/Plugins/Experimental/NNERuntimeIREE/Source/NNERuntimeIREE/Private/NNERuntimeIREEModel.h b/Engine/Plugins/Experimental/NNERuntimeIREE/Source/NNERuntimeIREE/Private/NNERuntimeIREEModel.h
--- a/Engine/Plugins/Experimental/NNERuntimeIREE/Source/NNERuntimeIREE/Private/NNERuntimeIREEModel.h
+++ b/Engine/Plugins/Experimental/NNERuntimeIREE/Source/NNERuntimeIREE/Private/NNERuntimeIREEModel.h
@@ -31,13 +31,16 @@ namespace UE::NNERuntimeIREE
 			static TSharedPtr<FModelInstance> Make(TSharedRef<UE::NNERuntimeIREE::CPU::Private::FDevice> InDevice, TSharedRef<UE::NNERuntimeIREE::Private::FModule> InModule);
 
 		public:
+			using ESetInputTensorShapeStatus = UE::NNE::IModelInstanceCPU::ESetInputTensorShapeStatus;
+			using ERunSyncStatus = UE::NNE::IModelInstanceCPU::ERunSyncStatus;
+
 			//~ Begin IModelInstanceCPU Interface
 			virtual TConstArrayView<UE::NNE::FTensorDesc> GetInputTensorDescs() const override;
 			virtual TConstArrayView<UE::NNE::FTensorDesc> GetOutputTensorDescs() const override;
 			virtual TConstArrayView<UE::NNE::FTensorShape> GetInputTensorShapes() const override;
 			virtual TConstArrayView<UE::NNE::FTensorShape> GetOutputTensorShapes() const override;
-			virtual int32 SetInputTensorShapes(TConstArrayView<UE::NNE::FTensorShape> InInputShapes) override;
-			virtual int32 RunSync(TConstArrayView<UE::NNE::FTensorBindingCPU> InInputBindings, TConstArrayView<UE::NNE::FTensorBindingCPU> InOutputBindings) override;
+			virtual ESetInputTensorShapeStatus SetInputTensorShapes(TConstArrayView<UE::NNE::FTensorShape> InInputShapes) override;
+			virtual ERunSyncStatus RunSync(TConstArrayView<UE::NNE::FTensorBindingCPU> InInputBindings, TConstArrayView<UE::NNE::FTensorBindingCPU> InOutputBindings) override;
 			//~ End IModelInstanceCPU Interface
 
 		private:
diff --git a/Engine/Plugins/Experimental/NNERuntimeORTCpu/Source/NNERuntimeORTCpu/Private/NNERuntimeORTCpuModel.cpp b/Engine/Plugins/Experimental/NNERuntimeORTCpu/Source/NNERuntimeORTCpu/Private/NNERuntimeORTCpuModel.cpp
--- a/Engine/Plugins/Experimental/NNERuntimeORTCpu/Source/NNERuntimeORTCpu/Private/NNERuntimeORTCpuModel.cpp
+++ b/Engine/Plugins/Experimental/NNERuntimeORTCpu/Source/NNERuntimeORTCpu/Private/NNERuntimeORTCpuModel.cpp
@@ -167,16 +167,16 @@ namespace UE::NNERuntimeORTCpu::Private
 		return true;
 	}
 	
-	int32 FModelInstanceCPU::SetInputTensorShapes(TConstArrayView<NNE::FTensorShape> InInputShapes)
+	FModelInstanceCPU::ESetInputTensorShapeStatus FModelInstanceCPU::SetInputTensorShapes(TConstArrayView<NNE::FTensorShape> InInputShapes)
 	{
 		InputTensors.Reset();
 		OutputTensors.Reset();
 		OutputTensorShapes.Reset();
 
 		// Verify input shape are valid for the model and set InputTensorShapes
-		if (FModelInstanceBase<IModelInstanceCPU>::SetInputTensorShapes(InInputShapes) != 0)
+		if (ESetInputTensorShapeStatus Status = FModelInstanceBase<IModelInstanceCPU>::SetInputTensorShapes(InInputShapes); Status != ESetInputTensorShapeStatus::Ok)
 		{
-			return -1;
+			return Status;
 		}
 
 		// Setup concrete input tensor
@@ -205,10 +205,10 @@ namespace UE::NNERuntimeORTCpu::Private
 			OutputTensorShapes.Reset();
 		}
 
-		return 0;
+		return ESetInputTensorShapeStatus::Ok;
 	}
 
-	int32 FModelInstanceCPU::RunSync(TConstArrayView<NNE::FTensorBindingCPU> InInputBindings, TConstArrayView<NNE::FTensorBindingCPU> InOutputBindings)
+	FModelInstanceCPU::ERunSyncStatus FModelInstanceCPU::RunSync(TConstArrayView<NNE::FTensorBindingCPU> InInputBindings, TConstArrayView<NNE::FTensorBindingCPU> InOutputBindings)
 	{
 		DECLARE_SCOPE_CYCLE_COUNTER(TEXT("FModelInstanceCPU_Run"), STAT_FModelInstanceCPU_Run, STATGROUP_NNE);
 
@@ -216,14 +216,14 @@ namespace UE::NNERuntimeORTCpu::Private
 		if (!bIsLoaded)
 		{
 			UE_LOG(LogNNE, Warning, TEXT("FModelInstanceCPU::Run(): Call FModelInstanceCPU::Load() to load a model first."));
-			return -1;
+			return ERunSyncStatus::Fail;
 		}
 
 		// Verify the model inputs were prepared
 		if (InputTensorShapes.Num() == 0)
 		{
 			UE_LOG(LogNNE, Error, TEXT("Run(): Input shapes are not set, please call SetInputTensorShapes."));
-			return -1;
+			return ERunSyncStatus::Fail;
 		}
 
 		NNEProfiling::Internal::FTimer RunTimer;
@@ -276,18 +276,18 @@ namespace UE::NNERuntimeORTCpu::Private
 		catch (const Ort::Exception& Exception)
 		{
 			UE_LOG(LogNNE, Error, TEXT("%s"), UTF8_TO_TCHAR(Exception.what()));
-			return false;
+			return ERunSyncStatus::Fail;
 		}
 		catch (...)
 		{
 			UE_LOG(LogNNE, Error, TEXT("Unknown exception!"));
-			return false;
+			return ERunSyncStatus::Fail;
 		}
 #endif //WITH_EDITOR
 
 		RunStatisticsEstimator.StoreSample(RunTimer.Toc());
 
-		return 0;
+		return ERunSyncStatus::Ok;
 	}
 
 	float FModelInstanceCPU::GetLastRunTimeMSec() const
diff --git a/Engine/Plugins/Experimental/NNERuntimeORTCpu/Source/NNERuntimeORTCpu/Private/NNERuntimeORTCpuModel.h b/Engine/Plugins/Experimental/NNERuntimeORTCpu/Source/NNERuntimeORTCpu/Private/NNERuntimeORTCpuModel.h
--- a/Engine/Plugins/Experimental/NNERuntimeORTCpu/Source/NNERuntimeORTCpu/Private/NNERuntimeORTCpuModel.h
+++ b/Engine/Plugins/Experimental/NNERuntimeORTCpu/Source/NNERuntimeORTCpu/Private/NNERuntimeORTCpuModel.h
@@ -30,12 +30,15 @@ namespace UE::NNERuntimeORTCpu::Private
 	{
 
 	public:
+		using ESetInputTensorShapeStatus = UE::NNE::IModelInstanceCPU::ESetInputTensorShapeStatus;
+		using ERunSyncStatus = UE::NNE::IModelInstanceCPU::ERunSyncStatus;
+
 		FModelInstanceCPU();
 		FModelInstanceCPU(Ort::Env* InORTEnvironment, const FRuntimeConf& InRuntimeConf);
 		virtual ~FModelInstanceCPU() {};
 
-		virtual int32 SetInputTensorShapes(TConstArrayView<NNE::FTensorShape> InInputShapes) override;
-		virtual int32 RunSync(TConstArrayView<NNE::FTensorBindingCPU> InInputBindings, TConstArrayView<NNE::FTensorBindingCPU> InOutputBindings) override;
+		virtual ESetInputTensorShapeStatus SetInputTensorShapes(TConstArrayView<NNE::FTensorShape> InInputShapes) override;
+		virtual ERunSyncStatus RunSync(TConstArrayView<NNE::FTensorBindingCPU> InInputBindings, TConstArrayView<NNE::FTensorBindingCPU> InOutputBindings) override;
 
 		bool Init(TConstArrayView<uint8> ModelData);
 		bool IsLoaded() const;
diff --git a/Engine/Plugins/Experimental/NNERuntimeRDG/Source/NNERuntimeRDG/Private/NNERuntimeRDGModel.cpp b/Engine/Plugins/Experimental/NNERuntimeRDG/Source/NNERuntimeRDG/Private/NNERuntimeRDGModel.cpp
--- a/Engine/Plugins/Experimental/NNERuntimeRDG/Source/NNERuntimeRDG/Private/NNERuntimeRDGModel.cpp
+++ b/Engine/Plugins/Experimental/NNERuntimeRDG/Source/NNERuntimeRDG/Private/NNERuntimeRDGModel.cpp
@@ -100,14 +100,14 @@ bool FModelInstanceRDG::LoadModel(TConstArrayView<uint8> ModelData, FNNERuntimeF
 
 
 
-int32 FModelInstanceRDG::SetInputTensorShapes(TConstArrayView<NNE::FTensorShape> InInputShapes)
+FModelInstanceRDG::ESetInputTensorShapeStatus FModelInstanceRDG::SetInputTensorShapes(TConstArrayView<NNE::FTensorShape> InInputShapes)
 {
 	OutputTensorShapes.Reset(OutputTensorIndices.Num());
 
 	//Verify input shape are valid for the model and set InputTensorShapes
-	if (FModelInstanceBase<NNE::IModelInstanceRDG>::SetInputTensorShapes(InInputShapes) != 0)
+	if (ESetInputTensorShapeStatus Status = FModelInstanceBase<NNE::IModelInstanceRDG>::SetInputTensorShapes(InInputShapes); Status != ESetInputTensorShapeStatus::Ok)
 	{
-		return -1;
+		return Status;
 	}
 
 	//Allocate and prime all AllTensorRDGRefs with concrete shapes defaulting variables dimension to 1 if needed
@@ -160,7 +160,7 @@ int32 FModelInstanceRDG::SetInputTensorShapes(TConstArrayView<NNE::FTensorShape>
 	//Allow the specific runtime to run shape inference if supported
 	if (PrepareTensorShapesAndData() != 0)
 	{
-		return -1;
+		return ESetInputTensorShapeStatus::Fail;
 	}
 
 	checkCode(
@@ -182,7 +182,7 @@ int32 FModelInstanceRDG::SetInputTensorShapes(TConstArrayView<NNE::FTensorShape>
 	check(WeightTensorIndices.Num() == WeightTensorRDGs.Num());
 	check(AllTensorRDGRefs.Num() == AllSymbolicTensorDescs.Num());
 	
-	return 0;
+	return ESetInputTensorShapeStatus::Ok;
 }
 
 FRDGBufferDesc CreateRDGBufferDescForTensorRDG(const FTensorRDG& Tensor)
@@ -195,7 +195,7 @@ FRDGBufferDesc CreateRDGBufferDescForTensorRDG(const FTensorRDG& Tensor)
 /**
  * Enqueue operators to RDG, the caller will run the GraphBuilder.Execute()
  */
-int32 FModelInstanceRDG::EnqueueRDG(FRDGBuilder& RDGBuilder, TConstArrayView<NNE::FTensorBindingRDG> InInputBindings, TConstArrayView<NNE::FTensorBindingRDG> InOutputBindings)
+FModelInstanceRDG::EEnqueueRDGStatus FModelInstanceRDG::EnqueueRDG(FRDGBuilder& RDGBuilder, TConstArrayView<NNE::FTensorBindingRDG> InInputBindings, TConstArrayView<NNE::FTensorBindingRDG> InOutputBindings)
 {
 	check(IsInRenderingThread());
 
@@ -205,21 +205,21 @@ int32 FModelInstanceRDG::EnqueueRDG(FRDGBuilder& RDGBuilder, TConstArrayView<NNE
 	if (InputTensorShapes.Num() == 0)
 	{
 		UE_LOG(LogNNE, Error, TEXT("EnqueueRDG(): Input shapes are not set, please call SetInputTensorShapes."));
-		return -1;
+		return EEnqueueRDGStatus::Fail;
 	}
 
 	Res = SetTensors(RDGBuilder, InputTensorRDGs, InInputBindings);
 	if (Res != -1)
 	{
 		UE_LOG(LogNNE, Warning, TEXT("Invalid buffer (was nullptr) for input tensor binding at index %d"), Res);
-		return -1;
+		return EEnqueueRDGStatus::Fail;
 	}
 
 	Res = SetTensors(RDGBuilder, OutputTensorRDGs, InOutputBindings);
 	if (Res != -1)
 	{
 		UE_LOG(LogNNE, Warning, TEXT("Invalid buffer (was nullptr) for output tensor binding at index %d"), Res);
-		return -1;
+		return EEnqueueRDGStatus::Fail;
 	}
 
 	
@@ -253,7 +253,7 @@ int32 FModelInstanceRDG::EnqueueRDG(FRDGBuilder& RDGBuilder, TConstArrayView<NNE
 	// We can now dispatch operators
 	AddDispatchOps_RenderThread(RDGBuilder);
 
-	return 0;
+	return EEnqueueRDGStatus::Ok;
 }
 
 int32 FModelInstanceRDG::SetTensors(FRDGBuilder& GraphBuilder, FTensorRDGArray& InTensorRDGs, TConstArrayView<NNE::FTensorBindingRDG> InBindings)
diff --git a/Engine/Plugins/Experimental/NNERuntimeRDG/Source/NNERuntimeRDG/Private/NNERuntimeRDGModel.h b/Engine/Plugins/Experimental/NNERuntimeRDG/Source/NNERuntimeRDG/Private/NNERuntimeRDGModel.h
--- a/Engine/Plugins/Experimental/NNERuntimeRDG/Source/NNERuntimeRDG/Private/NNERuntimeRDGModel.h
+++ b/Engine/Plugins/Experimental/NNERuntimeRDG/Source/NNERuntimeRDG/Private/NNERuntimeRDGModel.h
@@ -14,11 +14,14 @@ namespace UE::NNERuntimeRDG::Private
 	class FModelInstanceRDG : public NNE::Internal::FModelInstanceBase<NNE::IModelInstanceRDG>
 	{
 	public:
+		using ESetInputTensorShapeStatus = IModelInstanceRDG::ESetInputTensorShapeStatus;
+		using EEnqueueRDGStatus = IModelInstanceRDG::EEnqueueRDGStatus;
+
 		FModelInstanceRDG() {};
 		virtual ~FModelInstanceRDG() = default;
 
-		virtual int32 SetInputTensorShapes(TConstArrayView<NNE::FTensorShape> InInputShapes) override;
-		virtual int32 EnqueueRDG(FRDGBuilder& RDGBuilder, TConstArrayView<NNE::FTensorBindingRDG> InInputBindings, TConstArrayView<NNE::FTensorBindingRDG> InOutputBindings) override;
+		virtual ESetInputTensorShapeStatus SetInputTensorShapes(TConstArrayView<NNE::FTensorShape> InInputShapes) override;
+		virtual EEnqueueRDGStatus EnqueueRDG(FRDGBuilder& RDGBuilder, TConstArrayView<NNE::FTensorBindingRDG> InInputBindings, TConstArrayView<NNE::FTensorBindingRDG> InOutputBindings) override;
 		
 	protected:
 		bool LoadModel(TConstArrayView<uint8> ModelData, FNNERuntimeFormat& Format, int32 GuidAndVersionSize);
diff --git a/Engine/Plugins/Experimental/NeuralRendering/Source/NeuralPostProcessing/Private/NeuralPostProcessModelInstance.cpp b/Engine/Plugins/Experimental/NeuralRendering/Source/NeuralPostProcessing/Private/NeuralPostProcessModelInstance.cpp
--- a/Engine/Plugins/Experimental/NeuralRendering/Source/NeuralPostProcessing/Private/NeuralPostProcessModelInstance.cpp
+++ b/Engine/Plugins/Experimental/NeuralRendering/Source/NeuralPostProcessing/Private/NeuralPostProcessModelInstance.cpp
@@ -104,9 +104,9 @@ void UNeuralPostProcessModelInstance::Execute(FRDGBuilder& GraphBuilder)
 	auto RunSingleDispatch = [&](bool bAddFence = true) {
 		Input.Buffer = RDGInputBuffer;
 		Output.Buffer = RDGOutputBuffer;
-		int Status = ModelInstanceRDG->EnqueueRDG(GraphBuilder, InputBindings, OutputBindings);
+		UE::NNE::IModelInstanceRDG::EEnqueueRDGStatus Status = ModelInstanceRDG->EnqueueRDG(GraphBuilder, InputBindings, OutputBindings);
 
-		checkf(!Status, TEXT("EnqueueRDG failed: %d"), Status);
+		checkf(Status == UE::NNE::IModelInstanceRDG::EEnqueueRDGStatus::Ok, TEXT("EnqueueRDG failed: %d"), static_cast<int>(Status));
 
 		if (bAddFence)
 		{
@@ -295,7 +295,7 @@ void UNeuralPostProcessModelInstance::CreateDefaultNNEModel(UNNEModelData* NNEMo
 		// All dynamic dimensions are set to 1 by default.
 		ResolvedInputTensorShape = UE::NNE::FTensorShape::MakeFromSymbolic(InputShape);
 
-		if (ModelInstanceRDG->SetInputTensorShapes({ ResolvedInputTensorShape }) != 0)
+		if (ModelInstanceRDG->SetInputTensorShapes({ ResolvedInputTensorShape }) != UE::NNE::IModelInstanceRDG::ESetInputTensorShapeStatus::Ok)
 		{
 			ModelInstanceRDG.Reset();
 #if WITH_EDITOR
diff --git a/Engine/Plugins/NNE/NNERuntimeORT/Source/NNERuntimeORT/Private/NNERuntimeORTModel.cpp b/Engine/Plugins/NNE/NNERuntimeORT/Source/NNERuntimeORT/Private/NNERuntimeORTModel.cpp
--- a/Engine/Plugins/NNE/NNERuntimeORT/Source/NNERuntimeORT/Private/NNERuntimeORTModel.cpp
+++ b/Engine/Plugins/NNE/NNERuntimeORT/Source/NNERuntimeORT/Private/NNERuntimeORTModel.cpp
@@ -161,46 +161,48 @@ namespace UE::NNERuntimeORT::Private
 	}
 
 	template <class ModelInterface, class TensorBinding> 
-	int32 FModelInstanceORTBase<ModelInterface, TensorBinding>::SetInputTensorShapes(TConstArrayView<NNE::FTensorShape> InInputShapes)
+	typename ModelInterface::ESetInputTensorShapeStatus FModelInstanceORTBase<ModelInterface, TensorBinding>::SetInputTensorShapes(TConstArrayView<NNE::FTensorShape> InInputShapes)
 	{
+		using ModelInstanceBase = NNE::Internal::FModelInstanceBase<ModelInterface>;
+
 		InputTensors.Reset();
 		OutputTensors.Reset();
-		NNE::Internal::FModelInstanceBase<ModelInterface>::OutputTensorShapes.Reset();
+		ModelInstanceBase::OutputTensorShapes.Reset();
 
 		// Verify input shape are valid for the model and set InputTensorShapes
-		if (NNE::Internal::FModelInstanceBase<ModelInterface>::SetInputTensorShapes(InInputShapes) != 0)
+		if (typename ModelInterface::ESetInputTensorShapeStatus Status = ModelInstanceBase::SetInputTensorShapes(InInputShapes); Status != ModelInterface::ESetInputTensorShapeStatus::Ok)
 		{
-			return -1;
+			return Status;
 		}
 
 		// Setup concrete input tensor
-		for (int32 i = 0; i < NNE::Internal::FModelInstanceBase<ModelInterface>::InputSymbolicTensors.Num(); ++i)
+		for (int32 i = 0; i < ModelInstanceBase::InputSymbolicTensors.Num(); ++i)
 		{
-			NNE::Internal::FTensor Tensor = NNE::Internal::FTensor::Make(NNE::Internal::FModelInstanceBase<ModelInterface>::InputSymbolicTensors[i].GetName(), InInputShapes[i], NNE::Internal::FModelInstanceBase<ModelInterface>::InputSymbolicTensors[i].GetDataType());
+			NNE::Internal::FTensor Tensor = NNE::Internal::FTensor::Make(ModelInstanceBase::InputSymbolicTensors[i].GetName(), InInputShapes[i], ModelInstanceBase::InputSymbolicTensors[i].GetDataType());
 			InputTensors.Emplace(Tensor);
 		}
 
 		// Setup concrete output shapes only if all model output shapes are concretes, otherwise it will be set during Run()
-		for (NNE::FTensorDesc SymbolicTensorDesc : NNE::Internal::FModelInstanceBase<ModelInterface>::OutputSymbolicTensors)
+		for (NNE::FTensorDesc SymbolicTensorDesc : ModelInstanceBase::OutputSymbolicTensors)
 		{
 			if (SymbolicTensorDesc.GetShape().IsConcrete())
 			{
 				NNE::Internal::FTensor Tensor = NNE::Internal::FTensor::MakeFromSymbolicDesc(SymbolicTensorDesc);
 				OutputTensors.Emplace(Tensor);
-				NNE::Internal::FModelInstanceBase<ModelInterface>::OutputTensorShapes.Emplace(Tensor.GetShape());
+				ModelInstanceBase::OutputTensorShapes.Emplace(Tensor.GetShape());
 			}
 		}
-		if (OutputTensors.Num() != NNE::Internal::FModelInstanceBase<ModelInterface>::OutputSymbolicTensors.Num())
+		if (OutputTensors.Num() != ModelInstanceBase::OutputSymbolicTensors.Num())
 		{
 			OutputTensors.Reset();
-			NNE::Internal::FModelInstanceBase<ModelInterface>::OutputTensorShapes.Reset();
+			ModelInstanceBase::OutputTensorShapes.Reset();
 		}
 
-		return 0;
+		return ModelInterface::ESetInputTensorShapeStatus::Ok;
 	}
 
 	template <class ModelInterface, class TensorBinding>
-	int32 FModelInstanceORTBase<ModelInterface, TensorBinding>::RunSync(TConstArrayView<TensorBinding> InInputBindings, TConstArrayView<TensorBinding> InOutputBindings)
+	typename ModelInterface::ERunSyncStatus FModelInstanceORTBase<ModelInterface, TensorBinding>::RunSync(TConstArrayView<TensorBinding> InInputBindings, TConstArrayView<TensorBinding> InOutputBindings)
 	{
 		checkf(Session.IsValid(), TEXT("FModelInstanceORT::RunSync(): Called without a Session, FModelInstanceORT::Init() should have been called."));
 
@@ -210,7 +212,7 @@ namespace UE::NNERuntimeORT::Private
 		if (NNE::Internal::FModelInstanceBase<ModelInterface>::InputTensorShapes.Num() == 0)
 		{
 			UE_LOG(LogNNE, Error, TEXT("RunSync(): Input shapes are not set, please call SetInputTensorShapes."));
-			return -1;
+			return ModelInterface::ERunSyncStatus::Fail;
 		}
 
 		try
@@ -252,15 +254,15 @@ namespace UE::NNERuntimeORT::Private
 		catch (const Ort::Exception& Exception)
 		{
 			UE_LOG(LogNNE, Error, TEXT("%s"), UTF8_TO_TCHAR(Exception.what()));
-			return false;
+			return ModelInterface::ERunSyncStatus::Fail;
 		}
 		catch (...)
 		{
 			UE_LOG(LogNNE, Error, TEXT("Unknown exception!"));
-			return false;
+			return ModelInterface::ERunSyncStatus::Fail;
 		}
 
-		return 0;
+		return ModelInterface::ERunSyncStatus::Ok;
 	}
 
 	TSharedPtr<NNE::IModelInstanceCPU> FModelORTCpu::CreateModelInstanceCPU()
diff --git a/Engine/Plugins/NNE/NNERuntimeORT/Source/NNERuntimeORT/Private/NNERuntimeORTModel.h b/Engine/Plugins/NNE/NNERuntimeORT/Source/NNERuntimeORT/Private/NNERuntimeORTModel.h
--- a/Engine/Plugins/NNE/NNERuntimeORT/Source/NNERuntimeORT/Private/NNERuntimeORTModel.h
+++ b/Engine/Plugins/NNE/NNERuntimeORT/Source/NNERuntimeORT/Private/NNERuntimeORTModel.h
@@ -18,18 +18,19 @@ namespace UE::NNERuntimeORT::Private
 		GraphOptimizationLevel OptimizationLevel = GraphOptimizationLevel::ORT_ENABLE_ALL;
 	};
 
-	template <class ModelInterface, class TensorBinding> class FModelInstanceORTBase : public NNE::Internal::FModelInstanceBase<ModelInterface>
+	template <class ModelInterface, class TensorBinding>
+	class FModelInstanceORTBase : public NNE::Internal::FModelInstanceBase<ModelInterface>
 	{
 
 	public:
 		FModelInstanceORTBase(const FRuntimeConf& InRuntimeConf, TSharedPtr<Ort::Env> InEnvironment);
 		virtual ~FModelInstanceORTBase() = default;
 
-		virtual int32 SetInputTensorShapes(TConstArrayView<NNE::FTensorShape> InInputShapes) override;
+		virtual typename ModelInterface::ESetInputTensorShapeStatus SetInputTensorShapes(TConstArrayView<NNE::FTensorShape> InInputShapes) override;
 
 		bool Init(TConstArrayView<uint8> ModelData);
 
-		int32 RunSync(TConstArrayView<TensorBinding> InInputBindings, TConstArrayView<TensorBinding> InOutputBindings) override;
+		typename ModelInterface::ERunSyncStatus RunSync(TConstArrayView<TensorBinding> InInputBindings, TConstArrayView<TensorBinding> InOutputBindings) override;
 
 	protected:
 
diff --git a/Engine/Source/Runtime/NNE/Internal/NNEModelBase.h b/Engine/Source/Runtime/NNE/Internal/NNEModelBase.h
--- a/Engine/Source/Runtime/NNE/Internal/NNEModelBase.h
+++ b/Engine/Source/Runtime/NNE/Internal/NNEModelBase.h
@@ -8,56 +8,63 @@
 
 namespace UE::NNE::Internal
 {
-	template <class ModelInterface> class FModelInstanceBase : public ModelInterface
-	{
-	public:
 
-		virtual ~FModelInstanceBase() = default;
+template <class ModelInterface>
+class FModelInstanceBase : public ModelInterface
+{
+public:
+
+	virtual ~FModelInstanceBase() = default;
 
-		virtual TConstArrayView<NNE::FTensorDesc> GetInputTensorDescs() const override;
-		virtual TConstArrayView<NNE::FTensorDesc> GetOutputTensorDescs() const override;
-		virtual TConstArrayView<NNE::FTensorShape> GetInputTensorShapes() const override;
-		virtual TConstArrayView<NNE::FTensorShape> GetOutputTensorShapes() const override;
-		virtual int32 SetInputTensorShapes(TConstArrayView<NNE::FTensorShape> InInputShapes) override;
+	virtual TConstArrayView<NNE::FTensorDesc> GetInputTensorDescs() const override;
+	virtual TConstArrayView<NNE::FTensorDesc> GetOutputTensorDescs() const override;
+	virtual TConstArrayView<NNE::FTensorShape> GetInputTensorShapes() const override;
+	virtual TConstArrayView<NNE::FTensorShape> GetOutputTensorShapes() const override;
+	virtual typename ModelInterface::ESetInputTensorShapeStatus SetInputTensorShapes(TConstArrayView<NNE::FTensorShape> InInputShapes) override;
 
-	protected:
+protected:
 
-		FModelInstanceBase() {}
+	FModelInstanceBase() {}
 
-		TArray<FTensorShape>	InputTensorShapes;
-		TArray<FTensorShape>	OutputTensorShapes;
-		TArray<FTensorDesc>		InputSymbolicTensors;
-		TArray<FTensorDesc>		OutputSymbolicTensors;
-	};
+	TArray<FTensorShape>	InputTensorShapes;
+	TArray<FTensorShape>	OutputTensorShapes;
+	TArray<FTensorDesc>		InputSymbolicTensors;
+	TArray<FTensorDesc>		OutputSymbolicTensors;
+};
 
-	template <class T> TConstArrayView<FTensorDesc> FModelInstanceBase<T>::GetInputTensorDescs() const
-	{
-		return InputSymbolicTensors;
-	}
+template <class T>
+TConstArrayView<FTensorDesc> FModelInstanceBase<T>::GetInputTensorDescs() const
+{
+	return InputSymbolicTensors;
+}
 
-	template <class T> TConstArrayView<FTensorDesc> FModelInstanceBase<T>::GetOutputTensorDescs() const
-	{
-		return OutputSymbolicTensors;
-	}
+template <class T>
+TConstArrayView<FTensorDesc> FModelInstanceBase<T>::GetOutputTensorDescs() const
+{
+	return OutputSymbolicTensors;
+}
 
-	template <class T> TConstArrayView<FTensorShape> FModelInstanceBase<T>::GetInputTensorShapes() const
-	{
-		return InputTensorShapes;
-	}
+template <class T>
+TConstArrayView<FTensorShape> FModelInstanceBase<T>::GetInputTensorShapes() const
+{
+	return InputTensorShapes;
+}
 
-	template <class T> TConstArrayView<FTensorShape> FModelInstanceBase<T>::GetOutputTensorShapes() const
-	{
-		return OutputTensorShapes;
-	}
+template <class T>
+TConstArrayView<FTensorShape> FModelInstanceBase<T>::GetOutputTensorShapes() const
+{
+	return OutputTensorShapes;
+}
 
-	template <class T> int32 FModelInstanceBase<T>::SetInputTensorShapes(TConstArrayView<FTensorShape> InInputShapes)
-	{
-		InputTensorShapes.Reset(InInputShapes.Num());
+template <class T>
+typename T::ESetInputTensorShapeStatus FModelInstanceBase<T>::SetInputTensorShapes(TConstArrayView<FTensorShape> InInputShapes)
+{
+	InputTensorShapes.Reset(InInputShapes.Num());
 
 		if (InInputShapes.Num() != InputSymbolicTensors.Num())
 		{
 			UE_LOG(LogNNE, Error, TEXT("Number of input shapes does not match number of input tensors"));
-			return -1;
+			return T::ESetInputTensorShapeStatus::Fail;
 		}
 
 		for (int32 i = 0; i < InInputShapes.Num(); ++i)
@@ -66,15 +73,15 @@ namespace UE::NNE::Internal
 			if (!InInputShapes[i].IsCompatibleWith(SymbolicDesc.GetShape()))
 			{
 				UE_LOG(LogNNE, Error, TEXT("Input shape does not match input tensor %s of index %d"), *SymbolicDesc.GetName(), i);
-				return -1;
+				return T::ESetInputTensorShapeStatus::Fail;
 			}
 		}
+		
+	InputTensorShapes = InInputShapes;
 
-		InputTensorShapes = InInputShapes;
-
-		//Implementations are responsible to handle output and intermediate tensor shape inference.
-		//This base implementation only validate that all inputs are matching what the model can support.
-		return 0;
-	}
+	//Implementations are responsible to handle output and intermediate tensor shape inference.
+	//This base implementation only validate that all inputs are matching what the model can support.
+	return T::ESetInputTensorShapeStatus::Ok;
+}
 
 } // namespace UE::NNE::Internal
\ No newline at end of file
diff --git a/Engine/Source/Runtime/NNE/Public/NNERuntimeCPU.h b/Engine/Source/Runtime/NNE/Public/NNERuntimeCPU.h
--- a/Engine/Source/Runtime/NNE/Public/NNERuntimeCPU.h
+++ b/Engine/Source/Runtime/NNE/Public/NNERuntimeCPU.h
@@ -4,6 +4,7 @@
 
 #include "Containers/ContainersFwd.h"
 #include "NNETypes.h"
+#include "NNEStatus.h"
 #include "UObject/Interface.h"
 
 #include "NNERuntimeCPU.generated.h"
@@ -34,6 +35,9 @@ class NNE_API IModelInstanceCPU
 {
 public:
 
+	using ESetInputTensorShapeStatus = EResultStatus;
+	using ERunSyncStatus = EResultStatus;
+
 	virtual ~IModelInstanceCPU() = default;
 
 	/**
@@ -77,9 +81,9 @@ public:
 	 * This is a potentially expensive call and should be called lazily if possible.
 	 *
 	 * @param InInputShapes The input shapes to prepare the model with.
-	 * @return 0 on success or a non-zero number otherwise.
+	 * @return Status indicating success or failure.
 	 */
-	virtual int32 SetInputTensorShapes(TConstArrayView<FTensorShape> InInputShapes) = 0;
+	virtual ESetInputTensorShapeStatus SetInputTensorShapes(TConstArrayView<FTensorShape> InInputShapes) = 0;
 
 	/**
 	 * Evaluate the model synchronously.
@@ -91,9 +95,9 @@ public:
 	 *
 	 * @param InInputTensors An array containing tensor bindings for each input tensor with caller owned memory containing the input data.
 	 * @param InOutputTensors An array containing tensor bindings for each output tensor with caller owned memory big enough to contain the results on success.
-	 * @return 0 on success or a non-zero number otherwise.
+	 * @return Status indicating success or failure.
 	 */
-	virtual int32 RunSync(TConstArrayView<FTensorBindingCPU> InInputTensors, TConstArrayView<FTensorBindingCPU> InOutputTensors) = 0;
+	virtual ERunSyncStatus RunSync(TConstArrayView<FTensorBindingCPU> InInputTensors, TConstArrayView<FTensorBindingCPU> InOutputTensors) = 0;
 };
 
 /**
diff --git a/Engine/Source/Runtime/NNE/Public/NNERuntimeGPU.h b/Engine/Source/Runtime/NNE/Public/NNERuntimeGPU.h
--- a/Engine/Source/Runtime/NNE/Public/NNERuntimeGPU.h
+++ b/Engine/Source/Runtime/NNE/Public/NNERuntimeGPU.h
@@ -3,6 +3,7 @@
 #pragma once
 
 #include "Containers/ContainersFwd.h"
+#include "NNEStatus.h"
 #include "NNETypes.h"
 #include "UObject/Interface.h"
 
@@ -34,6 +35,9 @@ class NNE_API IModelInstanceGPU
 {
 public:
 
+	using ESetInputTensorShapeStatus = EResultStatus;
+	using ERunSyncStatus = EResultStatus;
+
 	virtual ~IModelInstanceGPU() = default;
 
 	/**
@@ -77,9 +81,9 @@ public:
 	 * This is a potentially expensive call and should be called lazily if possible.
 	 *
 	 * @param InInputShapes The input shapes to prepare the model with.
-	 * @return 0 on success or a non-zero number otherwise.
+	 * @return Status indicating success or failure.
 	 */
-	virtual int32 SetInputTensorShapes(TConstArrayView<FTensorShape> InInputShapes) = 0;
+	virtual ESetInputTensorShapeStatus SetInputTensorShapes(TConstArrayView<FTensorShape> InInputShapes) = 0;
 
 	/**
 	 * Evaluate the model synchronously.
@@ -91,9 +95,9 @@ public:
 	 *
 	 * @param InInputTensors An array containing tensor bindings for each input tensor with caller owned memory containing the input data.
 	 * @param InOutputTensors An array containing tensor bindings for each output tensor with caller owned memory big enough to contain the results on success.
-	 * @return 0 on success or a non-zero number otherwise.
+	 * @return Status indicating success or failure.
 	 */
-	virtual int32 RunSync(TConstArrayView<FTensorBindingGPU> InInputTensors, TConstArrayView<FTensorBindingGPU> InOutputTensors) = 0;
+	virtual ERunSyncStatus RunSync(TConstArrayView<FTensorBindingGPU> InInputTensors, TConstArrayView<FTensorBindingGPU> InOutputTensors) = 0;
 };
 
 /**
diff --git a/Engine/Source/Runtime/NNE/Public/NNERuntimeRDG.h b/Engine/Source/Runtime/NNE/Public/NNERuntimeRDG.h
--- a/Engine/Source/Runtime/NNE/Public/NNERuntimeRDG.h
+++ b/Engine/Source/Runtime/NNE/Public/NNERuntimeRDG.h
@@ -3,6 +3,7 @@
 #pragma once
 
 #include "Containers/ContainersFwd.h"
+#include "NNEStatus.h"
 #include "NNETypes.h"
 #include "RenderGraphFwd.h"
 #include "UObject/Interface.h"
@@ -34,6 +35,9 @@ class NNE_API IModelInstanceRDG
 {
 public:
 
+	using ESetInputTensorShapeStatus = EResultStatus;
+	using EEnqueueRDGStatus = EResultStatus;
+
 	virtual ~IModelInstanceRDG() = default;
 
 	/**
@@ -78,9 +82,9 @@ public:
 	 * This is a potentially expensive call and should be called lazily if possible.
 	 *
 	 * @param InInputShapes The input shapes to prepare the model with.
-	 * @return 0 on success or a non-zero number otherwise.
+	 * @return Status indicating success or failure.
 	 */
-	virtual int32 SetInputTensorShapes(TConstArrayView<FTensorShape> InInputShapes) = 0;
+	virtual ESetInputTensorShapeStatus SetInputTensorShapes(TConstArrayView<FTensorShape> InInputShapes) = 0;
 
 	/**
 	 * Enqueue the model graph to a FRDGBuilder.
@@ -92,9 +96,9 @@ public:
 	 * @param RDGBuilder The RDG builder to which the neural network operations are enqueued.
 	 * @param InInputTensors An array containing tensor bindings for each input tensor with caller owned memory containing the input data.
 	 * @param InOutputTensors An array containing tensor bindings for each output tensor with caller owned memory big enough to contain the results on success.
-	 * @return 0 on success or a non-zero number otherwise.
+	 * @return Status indicating success or failure.
 	 */
-	virtual int32 EnqueueRDG(FRDGBuilder& RDGBuilder, TConstArrayView<FTensorBindingRDG> Inputs, TConstArrayView<FTensorBindingRDG> Outputs) = 0;
+	virtual EEnqueueRDGStatus EnqueueRDG(FRDGBuilder& RDGBuilder, TConstArrayView<FTensorBindingRDG> Inputs, TConstArrayView<FTensorBindingRDG> Outputs) = 0;
 };
 
 /**
diff --git a/Engine/Source/Runtime/NNE/Public/NNEStatus.h b/Engine/Source/Runtime/NNE/Public/NNEStatus.h
new file mode 100644
--- /dev/null
+++ b/Engine/Source/Runtime/NNE/Public/NNEStatus.h
@@ -0,0 +1,16 @@
+// Copyright Epic Games, Inc. All Rights Reserved.
+
+#pragma once
+
+#include "HAL/Platform.h"
+
+namespace UE::NNE
+{
+
+enum class EResultStatus : uint8
+{
+	Ok = 0,
+	Fail = 1
+};
+
+} // namespace UE::NNE
