[NNE] NNERuntimeORTDml disable memory pattern as it is not supported by the DirectML execution provider.

#tests Editor (Win64), smoke tests

#rnx

diff --git a/Engine/Plugins/NNE/NNERuntimeORT/Source/NNERuntimeORT/Private/NNERuntimeORTModel.cpp b/Engine/Plugins/NNE/NNERuntimeORT/Source/NNERuntimeORT/Private/NNERuntimeORTModel.cpp
--- a/Engine/Plugins/NNE/NNERuntimeORT/Source/NNERuntimeORT/Private/NNERuntimeORTModel.cpp
+++ b/Engine/Plugins/NNE/NNERuntimeORT/Source/NNERuntimeORT/Private/NNERuntimeORTModel.cpp
@@ -328,6 +328,8 @@ namespace UE::NNERuntimeORT::Private
 			return false;
 		}
 
+		SessionOptions->SetExecutionMode(ExecutionMode::ORT_SEQUENTIAL);
+		SessionOptions->DisableMemPattern();
 		SessionOptions->DisableCpuMemArena();
 
 		// In order to use DirectML we need D3D12
