[NNE] NNERuntimeORTDml adapt to ORT API change.

Problem:
The ORT API call to append the DirectML execution provider to the ORT session options is deprecated. The is a more specialized API for the DirectML execution provider (OrtDmlApi).

Solution:
Use the dedicated API of the DirectML execution provider.

#tests Editor (Win64) smoke tests

#rnx

diff --git a/Engine/Plugins/NNE/NNERuntimeORT/Source/NNERuntimeORT/Private/NNERuntimeORTModel.cpp b/Engine/Plugins/NNE/NNERuntimeORT/Source/NNERuntimeORT/Private/NNERuntimeORTModel.cpp
--- a/Engine/Plugins/NNE/NNERuntimeORT/Source/NNERuntimeORT/Private/NNERuntimeORTModel.cpp
+++ b/Engine/Plugins/NNE/NNERuntimeORT/Source/NNERuntimeORT/Private/NNERuntimeORTModel.cpp
@@ -383,13 +383,22 @@ namespace UE::NNERuntimeORT::Private
 
 		if (FAILED(Res) || !DmlDevice)
 		{
-			UE_LOG(LogNNE, Warning, TEXT("Failed to create DirectML device, DMLCreateDevice error code :%x"), Res);
+			UE_LOG(LogNNE, Error, TEXT("Failed to create DirectML device, DMLCreateDevice error code :%x"), Res);
 			return false;
 		}
 
 		ID3D12CommandQueue* CmdQ = RHI->RHIGetCommandQueue();
-		OrtStatusPtr Status = OrtSessionOptionsAppendExecutionProviderEx_DML(*SessionOptions.Get(), DmlDevice, CmdQ);
+		
+		const OrtDmlApi* DmlApi = nullptr;
+		Ort::ThrowOnError(Ort::GetApi().GetExecutionProviderApi("DML", ORT_API_VERSION, reinterpret_cast<const void**>(&DmlApi)));
+
+		if (!DmlApi)
+		{
+			UE_LOG(LogNNE, Error, TEXT("Ort DirectML Api not available!"));
+			return false;
+		}
 
+		OrtStatusPtr Status = DmlApi->SessionOptionsAppendExecutionProvider_DML1(*SessionOptions.Get(), DmlDevice, CmdQ);
 		if (Status)
 		{
 			UE_LOG(LogNNE, Error, TEXT("Failed to add DirectML execution provider to OnnxRuntime session options: %s"), ANSI_TO_TCHAR(Ort::GetApi().GetErrorMessage(Status)));
