[NNE] ORTCpu optimizer output ort format for faster model loading.

Problem: NNERuntimeORTCpu was optimizing the model to the ONNX format at cook time skipping some potential optimization that can only be expressed in the ORT and that thus would happen at session creation.

Solution: Optimize the model to the ORT format at cook time for the ORTCpu runtime.

Test: All tests win64 editor.

#rnx

diff --git a/Engine/Plugins/NNE/NNERuntimeORT/Source/NNERuntimeORT/Private/NNERuntimeORT.cpp b/Engine/Plugins/NNE/NNERuntimeORT/Source/NNERuntimeORT/Private/NNERuntimeORT.cpp
--- a/Engine/Plugins/NNE/NNERuntimeORT/Source/NNERuntimeORT/Private/NNERuntimeORT.cpp
+++ b/Engine/Plugins/NNE/NNERuntimeORT/Source/NNERuntimeORT/Private/NNERuntimeORT.cpp
@@ -23,7 +23,7 @@ FGuid UNNERuntimeORTDml::GUID = FGuid((int32)'O', (int32)'G', (int32)'P', (int32
 int32 UNNERuntimeORTDml::Version = 0x00000001;
 
 FGuid UNNERuntimeORTCpu::GUID = FGuid((int32)'O', (int32)'C', (int32)'P', (int32)'U');
-int32 UNNERuntimeORTCpu::Version = 0x00000001;
+int32 UNNERuntimeORTCpu::Version = 0x00000002;
 
 UNNERuntimeORTDml::ECanCreateModelDataStatus UNNERuntimeORTDml::CanCreateModelData(const FString& FileType, TConstArrayView<uint8> FileData, const TMap<FString, TConstArrayView<uint8>>& AdditionalFileData, const FGuid& FileId, const ITargetPlatform* TargetPlatform) const
 {
@@ -72,7 +72,7 @@ TSharedPtr<UE::NNE::FSharedModelData> UNNERuntimeORTCpu::CreateModelData(const F
 		return {};
 	}
 
-	TUniquePtr<UE::NNE::Internal::IModelOptimizer> Optimizer = UE::NNEUtilities::Internal::CreateONNXToONNXModelOptimizer();
+	TUniquePtr<UE::NNE::Internal::IModelOptimizer> Optimizer = UE::NNEUtilities::Internal::CreateONNXToORTModelOptimizer();
 
 	FNNEModelRaw InputModel;
 	InputModel.Data = FileData;
