[NNE] Fix for ClangEditor build.

#tests Local compilation ClangEditor Win64

#rnx

diff --git a/Engine/Plugins/NNE/NNERuntimeORT/Source/NNERuntimeORT/Private/NNERuntimeORTModel.cpp b/Engine/Plugins/NNE/NNERuntimeORT/Source/NNERuntimeORT/Private/NNERuntimeORTModel.cpp
--- a/Engine/Plugins/NNE/NNERuntimeORT/Source/NNERuntimeORT/Private/NNERuntimeORTModel.cpp
+++ b/Engine/Plugins/NNE/NNERuntimeORT/Source/NNERuntimeORT/Private/NNERuntimeORTModel.cpp
@@ -834,7 +834,7 @@ FModelInstanceORTDmlRDG::EEnqueueRDGStatus FModelInstanceORTDmlRDG::EnqueueRDG(F
 					TArray<Ort::Value> OrtInputTensors;
 					TArray<Ort::Value> OrtOutputTensors;
 
-					Ort::IoBinding IoBinding = Ort::IoBinding::IoBinding(*Session);
+					Ort::IoBinding IoBinding = Ort::IoBinding(*Session);
 
 					for (int32 i = 0; i < InputBuffersCopyCopy.Num(); i++)
 					{
