[NNE] Enable NNERuntimeORTCpu and NNERuntimeORTDml (GPU) in standalone.

#tests Editor/Game, Win64 reimport and all tests

#rnx

diff --git a/Engine/Plugins/NNE/NNERuntimeORT/NNERuntimeORT.uplugin b/Engine/Plugins/NNE/NNERuntimeORT/NNERuntimeORT.uplugin
--- a/Engine/Plugins/NNE/NNERuntimeORT/NNERuntimeORT.uplugin
+++ b/Engine/Plugins/NNE/NNERuntimeORT/NNERuntimeORT.uplugin
@@ -19,7 +19,7 @@
 		{
 			"Name": "NNERuntimeORT",
 			"Type": "RuntimeAndProgram",
-			"LoadingPhase": "Default",
+			"LoadingPhase": "PreDefault",
 			"PlatformAllowList": [ "Win64", "Linux", "Mac" ]
 		}
 	]
diff --git a/Engine/Plugins/NNE/NNERuntimeORT/Source/NNERuntimeORT/NNERuntimeORT.Build.cs b/Engine/Plugins/NNE/NNERuntimeORT/Source/NNERuntimeORT/NNERuntimeORT.Build.cs
--- a/Engine/Plugins/NNE/NNERuntimeORT/Source/NNERuntimeORT/NNERuntimeORT.Build.cs
+++ b/Engine/Plugins/NNE/NNERuntimeORT/Source/NNERuntimeORT/NNERuntimeORT.Build.cs
@@ -8,25 +8,23 @@ public class NNERuntimeORT : ModuleRules
 	{
 		PCHUsage = PCHUsageMode.UseExplicitOrSharedPCHs;
 
-		PrivateDependencyModuleNames.AddRange
-			(
-			new string[] {
-				"Core",
-				"CoreUObject",
-				"Engine",
-				"NNE",
-				"NNEOnnxruntime",
-				"Projects",
-				"RenderCore"
-			}
-		);
+		PrivateDependencyModuleNames.AddRange(new string[]
+		{
+			"Core",
+			"CoreUObject",
+			"Engine",
+			"NNE",
+			"NNEOnnxruntime",
+			"Projects",
+			"RenderCore",
+			"DeveloperSettings"
+		});
 
 		if (Target.Platform == UnrealTargetPlatform.Win64)
 		{
 			PrivateDependencyModuleNames.AddRange(new string[]
 			{
 				"D3D12RHI",
-				"DirectML",
 				"RHI"
 			});
 
@@ -35,7 +33,6 @@ public class NNERuntimeORT : ModuleRules
 				"DirectML",
 				"DX12"
 			});
-
 		}
 	}
 }
diff --git a/Engine/Plugins/NNE/NNERuntimeORT/Source/NNERuntimeORT/Private/NNERuntimeORT.cpp b/Engine/Plugins/NNE/NNERuntimeORT/Source/NNERuntimeORT/Private/NNERuntimeORT.cpp
--- a/Engine/Plugins/NNE/NNERuntimeORT/Source/NNERuntimeORT/Private/NNERuntimeORT.cpp
+++ b/Engine/Plugins/NNE/NNERuntimeORT/Source/NNERuntimeORT/Private/NNERuntimeORT.cpp
@@ -21,187 +21,8 @@
 FGuid UNNERuntimeORTCpu::GUID = FGuid((int32)'O', (int32)'C', (int32)'P', (int32)'U');
 int32 UNNERuntimeORTCpu::Version = 0x00000002;
 
-FGuid UNNERuntimeORTDmlEditor::GUID = FGuid((int32)'O', (int32)'D', (int32)'M', (int32)'L');
-int32 UNNERuntimeORTDmlEditor::Version = 0x00000002;
-
-FGuid UNNERuntimeORTDml::GUID = UNNERuntimeORTDmlEditor::GUID;
-int32 UNNERuntimeORTDml::Version = UNNERuntimeORTDmlEditor::Version;
-
-namespace UE::NNERuntimeORT::Private
-{
-
-class NNERuntimeORTDmlImpl : public INNERuntime, public INNERuntimeGPU, public INNERuntimeRDG
-{
-	using ECanCreateModelCommonStatus = UE::NNE::EResultStatus;
-
-private:
-	TSharedPtr<Ort::Env> ORTEnvironment;
-	FGuid GUID;
-	int32 Version;
-	
-public:
-	NNERuntimeORTDmlImpl(FGuid GUID, int32 Version) : GUID(GUID), Version(Version)
-	{
-
-	}
-
-	virtual ~NNERuntimeORTDmlImpl() = default;
-
-	void Init()
-	{
-		check(!ORTEnvironment.IsValid());
-		ORTEnvironment = MakeShared<Ort::Env>();
-	}
-
-	virtual FString GetRuntimeName() const override
-	{
-		return TEXT("NNERuntimeORTDml");
-	}
-
-	virtual ECanCreateModelDataStatus CanCreateModelData(const FString& FileType, TConstArrayView<uint8> FileData, const TMap<FString, TConstArrayView<uint8>>& AdditionalFileData, const FGuid& FileId, const ITargetPlatform* TargetPlatform) const override
-	{
-		return (!FileData.IsEmpty() && FileType.Compare("onnx", ESearchCase::IgnoreCase) == 0) ? ECanCreateModelDataStatus::Ok : ECanCreateModelDataStatus::FailFileIdNotSupported;
-	}
-
-	virtual TSharedPtr<UE::NNE::FSharedModelData> CreateModelData(const FString& FileType, TConstArrayView<uint8> FileData, const TMap<FString, TConstArrayView<uint8>>& AdditionalFileData, const FGuid& FileId, const ITargetPlatform* TargetPlatform) override
-	{
-		if (CanCreateModelData(FileType, FileData, AdditionalFileData, FileId, TargetPlatform) != ECanCreateModelDataStatus::Ok)
-		{
-			UE_LOG(LogNNE, Warning, TEXT("NNERuntimeORTDmlImpl cannot create the model data with id %s (Filetype: %s)"), *FileId.ToString(EGuidFormats::Digits).ToLower(), *FileType);
-			return {};
-		}
-
-		FNNEModelRaw InputModel{TArray<uint8>{FileData}, ENNEInferenceFormat::ONNX};
-		if (!UE::NNERuntimeORT::Private::OrtHelper::OptimizeModel(InputModel, ENNEInferenceFormat::ONNX))
-		{
-			return {};
-		}
-
-		TArray<uint8> Result;
-		FMemoryWriter Writer(Result);
-		Writer << UNNERuntimeORTDml::GUID;
-		Writer << UNNERuntimeORTDml::Version;
-		Writer.Serialize(InputModel.Data.GetData(), InputModel.Data.Num());
-
-		return MakeShared<UE::NNE::FSharedModelData>(MakeSharedBufferFromArray(MoveTemp(Result)), 0);
-	}
-
-	virtual FString GetModelDataIdentifier(const FString& FileType, TConstArrayView<uint8> FileData, const TMap<FString, TConstArrayView<uint8>>& AdditionalFileData, const FGuid& FileId, const ITargetPlatform* TargetPlatform) const override
-	{
-		return FileId.ToString(EGuidFormats::Digits) + "-" + UNNERuntimeORTDml::GUID.ToString(EGuidFormats::Digits) + "-" + FString::FromInt(UNNERuntimeORTDml::Version);
-	}
-
-	virtual ECanCreateModelGPUStatus CanCreateModelGPU(const TObjectPtr<UNNEModelData> ModelData) const override
-	{
-		return CanCreateModelCommon(ModelData) == ECanCreateModelCommonStatus::Ok ? ECanCreateModelGPUStatus::Ok : ECanCreateModelGPUStatus::Fail;
-	}
-
-	virtual TSharedPtr<UE::NNE::IModelGPU> CreateModelGPU(const TObjectPtr<UNNEModelData> ModelData) override
-	{
-#if PLATFORM_WINDOWS
-		check(ModelData);
-		check(ORTEnvironment.IsValid());
-
-		if (CanCreateModelGPU(ModelData) != ECanCreateModelGPUStatus::Ok)
-		{
-			UE_LOG(LogNNE, Warning, TEXT("UNNERuntimeORTDml cannot create a model GPU from the model data with id %s"), *ModelData->GetFileId().ToString(EGuidFormats::Digits));
-			return {};
-		}
-
-		const TSharedPtr<UE::NNE::FSharedModelData> SharedData = ModelData->GetModelData(GetRuntimeName());
-		check(SharedData.IsValid());
-
-		if (FEngineAnalytics::IsAvailable())
-		{
-			TArray<FAnalyticsEventAttribute> Attributes = MakeAnalyticsEventAttributeArray(
-				TEXT("PlatformName"), UGameplayStatics::GetPlatformName(),
-				TEXT("HashedRuntimeName"), FMD5::HashAnsiString(*GetRuntimeName()),
-				TEXT("ModelDataSize"), SharedData->GetView().Num()
-			);
-			FEngineAnalytics::GetProvider().RecordEvent(TEXT("NeuralNetworkEngine.CreateModel"), Attributes);
-		}
-
-		return MakeShared<UE::NNERuntimeORT::Private::FModelORTDmlGPU>(ORTEnvironment, SharedData);
-#else // PLATFORM_WINDOWS
-		return TSharedPtr<UE::NNE::IModelGPU>();
-#endif // PLATFORM_WINDOWS
-	}
-
-	virtual ECanCreateModelRDGStatus CanCreateModelRDG(TObjectPtr<UNNEModelData> ModelData) const override{
-		return CanCreateModelCommon(ModelData) == ECanCreateModelCommonStatus::Ok ? ECanCreateModelRDGStatus::Ok : ECanCreateModelRDGStatus::Fail;
-	}
-
-	virtual TSharedPtr<UE::NNE::IModelRDG> CreateModelRDG(TObjectPtr<UNNEModelData> ModelData) override{
-#if PLATFORM_WINDOWS
-		check(ModelData);
-		check(ORTEnvironment.IsValid());
-
-		if (CanCreateModelRDG(ModelData) != ECanCreateModelRDGStatus::Ok)
-		{
-			UE_LOG(LogNNE, Warning, TEXT("UNNERuntimeORTDml cannot create a model RDG from the model data with id %s"), *ModelData->GetFileId().ToString(EGuidFormats::Digits));
-			return {};
-		}
-
-		const TSharedRef<UE::NNE::FSharedModelData> SharedData = ModelData->GetModelData(GetRuntimeName()).ToSharedRef();
-
-		if (FEngineAnalytics::IsAvailable())
-		{
-			TArray<FAnalyticsEventAttribute> Attributes = MakeAnalyticsEventAttributeArray(
-				TEXT("PlatformName"), UGameplayStatics::GetPlatformName(),
-				TEXT("HashedRuntimeName"), FMD5::HashAnsiString(*GetRuntimeName()),
-				TEXT("ModelDataSize"), SharedData->GetView().Num()
-			);
-			FEngineAnalytics::GetProvider().RecordEvent(TEXT("NeuralNetworkEngine.CreateModel"), Attributes);
-		}
-
-		return MakeShared<UE::NNERuntimeORT::Private::FModelORTDmlRDG>(ORTEnvironment.ToSharedRef(), SharedData);
-#else // PLATFORM_MICROSOFT
-		return {};
-#endif // PLATFORM_MICROSOFT
-	}
-
-private:
-	ECanCreateModelCommonStatus CanCreateModelCommon(const TObjectPtr<UNNEModelData> ModelData) const
-	{
-#if PLATFORM_WINDOWS
-		check(ModelData != nullptr);
-
-		// In order to use DirectML we need D3D12
-		if (!IsRHID3D12())
-		{
-			return ECanCreateModelCommonStatus::Fail;
-		}
-
-		constexpr int32 GuidSize = sizeof(UNNERuntimeORTDml::GUID);
-		constexpr int32 VersionSize = sizeof(UNNERuntimeORTDml::Version);
-		const TSharedPtr<UE::NNE::FSharedModelData> SharedData = ModelData->GetModelData(GetRuntimeName());
-
-		if (!SharedData.IsValid())
-		{
-			return ECanCreateModelCommonStatus::Fail;
-		}
-
-		TConstArrayView<uint8> Data = SharedData->GetView();
-
-		if (Data.Num() <= GuidSize + VersionSize)
-		{
-			return ECanCreateModelCommonStatus::Fail;
-		}
-
-		static const FGuid DeprecatedGUID = FGuid((int32)'O', (int32)'G', (int32)'P', (int32)'U');
-
-		bool bResult = FGenericPlatformMemory::Memcmp(&(Data[0]), &(UNNERuntimeORTDml::GUID), GuidSize) == 0;
-		bResult |= FGenericPlatformMemory::Memcmp(&(Data[0]), &(DeprecatedGUID), GuidSize) == 0;
-		bResult &= FGenericPlatformMemory::Memcmp(&(Data[GuidSize]), &(UNNERuntimeORTDml::Version), VersionSize) == 0;
-
-		return bResult ? ECanCreateModelCommonStatus::Ok : ECanCreateModelCommonStatus::Fail;
-#else // PLATFORM_MICROSOFT
-		return ECanCreateModelCommonStatus::Fail;
-#endif // PLATFORM_MICROSOFT
-	}
-};
-
-} // namespace UE::NNERuntimeORT::Private
+FGuid UNNERuntimeORTDml::GUID = FGuid((int32)'O', (int32)'D', (int32)'M', (int32)'L');
+int32 UNNERuntimeORTDml::Version = 0x00000002;
 
 UNNERuntimeORTCpu::ECanCreateModelDataStatus UNNERuntimeORTCpu::CanCreateModelData(const FString& FileType, TConstArrayView<uint8> FileData, const TMap<FString, TConstArrayView<uint8>>& AdditionalFileData, const FGuid& FileId, const ITargetPlatform* TargetPlatform) const
 {
@@ -236,12 +57,9 @@ FString UNNERuntimeORTCpu::GetModelDataIdentifier(const FString& FileType, TCons
 	return FileId.ToString(EGuidFormats::Digits) + "-" + UNNERuntimeORTCpu::GUID.ToString(EGuidFormats::Digits) + "-" + FString::FromInt(UNNERuntimeORTCpu::Version);
 }
 
-void UNNERuntimeORTCpu::Init()
+void UNNERuntimeORTCpu::Init(TSharedRef<UE::NNERuntimeORT::Private::FEnvironment> InEnvironment)
 {
-#if WITH_EDITOR
-	check(!ORTEnvironment.IsValid());
-	ORTEnvironment = MakeShared<Ort::Env>();
-#endif // WITH_EDITOR
+	Environment = InEnvironment;
 }
 
 FString UNNERuntimeORTCpu::GetRuntimeName() const
@@ -251,7 +69,6 @@ FString UNNERuntimeORTCpu::GetRuntimeName() const
 
 UNNERuntimeORTCpu::ECanCreateModelCPUStatus UNNERuntimeORTCpu::CanCreateModelCPU(const TObjectPtr<UNNEModelData> ModelData) const
 {
-#if WITH_EDITOR
 	check(ModelData != nullptr);
 
 	constexpr int32 GuidSize = sizeof(UNNERuntimeORTCpu::GUID);
@@ -274,16 +91,11 @@ UNNERuntimeORTCpu::ECanCreateModelCPUStatus UNNERuntimeORTCpu::CanCreateModelCPU
 	bResult &= FGenericPlatformMemory::Memcmp(&(Data[GuidSize]), &(UNNERuntimeORTCpu::Version), VersionSize) == 0;
 
 	return bResult ? ECanCreateModelCPUStatus::Ok : ECanCreateModelCPUStatus::Fail;
-#else
-	return ECanCreateModelCPUStatus::Fail;
-#endif // WITH_EDITOR
 }
 
 TSharedPtr<UE::NNE::IModelCPU> UNNERuntimeORTCpu::CreateModelCPU(const TObjectPtr<UNNEModelData> ModelData)
 {
-#if WITH_EDITOR
 	check(ModelData != nullptr);
-	check(ORTEnvironment.IsValid());
 
 	if (CanCreateModelCPU(ModelData) != ECanCreateModelCPUStatus::Ok)
 	{
@@ -294,7 +106,7 @@ TSharedPtr<UE::NNE::IModelCPU> UNNERuntimeORTCpu::CreateModelCPU(const TObjectPt
 	const TSharedPtr<UE::NNE::FSharedModelData> SharedData = ModelData->GetModelData(GetRuntimeName());
 	check(SharedData.IsValid());
 
-	UE::NNE::IModelCPU* IModel = static_cast<UE::NNE::IModelCPU*>(new UE::NNERuntimeORT::Private::FModelORTCpu(ORTEnvironment, SharedData));
+	UE::NNE::IModelCPU* IModel = static_cast<UE::NNE::IModelCPU*>(new UE::NNERuntimeORT::Private::FModelORTCpu(Environment.ToSharedRef(), SharedData));
 	check(IModel != nullptr);
 
 	if (FEngineAnalytics::IsAvailable())
@@ -308,103 +120,158 @@ TSharedPtr<UE::NNE::IModelCPU> UNNERuntimeORTCpu::CreateModelCPU(const TObjectPt
 	}
 
 	return TSharedPtr<UE::NNE::IModelCPU>(IModel);
-#else
-	return {};
-#endif // WITH_EDITOR
 }
 
 /*
- * UNNERuntimeORTDmlEditor
+ * UNNERuntimeORTDml
  */
-UNNERuntimeORTDmlEditor::UNNERuntimeORTDmlEditor()
+void UNNERuntimeORTDml::Init(TSharedRef<UE::NNERuntimeORT::Private::FEnvironment> InEnvironment)
 {
-	Impl = MakeUnique<UE::NNERuntimeORT::Private::NNERuntimeORTDmlImpl>(GUID, Version);
+	Environment = InEnvironment;
 }
 
-void UNNERuntimeORTDmlEditor::Init()
+FString UNNERuntimeORTDml::GetRuntimeName() const
 {
-	Impl->Init();
+	return TEXT("NNERuntimeORTDml");
 }
 
-FString UNNERuntimeORTDmlEditor::GetRuntimeName() const
+UNNERuntimeORTDml::ECanCreateModelDataStatus UNNERuntimeORTDml::CanCreateModelData(const FString& FileType, TConstArrayView<uint8> FileData, const TMap<FString, TConstArrayView<uint8>>& AdditionalFileData, const FGuid& FileId, const ITargetPlatform* TargetPlatform) const
 {
-	return Impl->GetRuntimeName();
+	return (!FileData.IsEmpty() && FileType.Compare("onnx", ESearchCase::IgnoreCase) == 0) ? ECanCreateModelDataStatus::Ok : ECanCreateModelDataStatus::FailFileIdNotSupported;
 }
 
-UNNERuntimeORTDmlEditor::ECanCreateModelDataStatus UNNERuntimeORTDmlEditor::CanCreateModelData(const FString& FileType, TConstArrayView<uint8> FileData, const TMap<FString, TConstArrayView<uint8>>& AdditionalFileData, const FGuid& FileId, const ITargetPlatform* TargetPlatform) const
+TSharedPtr<UE::NNE::FSharedModelData> UNNERuntimeORTDml::CreateModelData(const FString& FileType, TConstArrayView<uint8> FileData, const TMap<FString, TConstArrayView<uint8>>& AdditionalFileData, const FGuid& FileId, const ITargetPlatform* TargetPlatform)
 {
-	return Impl->CanCreateModelData(FileType, FileData, AdditionalFileData, FileId, TargetPlatform);
-}
+	if (CanCreateModelData(FileType, FileData, AdditionalFileData, FileId, TargetPlatform) != ECanCreateModelDataStatus::Ok)
+	{
+		UE_LOG(LogNNE, Warning, TEXT("NNERuntimeORTDmlImpl cannot create the model data with id %s (Filetype: %s)"), *FileId.ToString(EGuidFormats::Digits).ToLower(), *FileType);
+		return {};
+	}
 
-TSharedPtr<UE::NNE::FSharedModelData> UNNERuntimeORTDmlEditor::CreateModelData(const FString& FileType, TConstArrayView<uint8> FileData, const TMap<FString, TConstArrayView<uint8>>& AdditionalFileData, const FGuid& FileId, const ITargetPlatform* TargetPlatform)
-{
-	return Impl->CreateModelData(FileType, FileData, AdditionalFileData, FileId, TargetPlatform);
-}
+	FNNEModelRaw InputModel{TArray<uint8>{FileData}, ENNEInferenceFormat::ONNX};
+	if (!UE::NNERuntimeORT::Private::OrtHelper::OptimizeModel(InputModel, ENNEInferenceFormat::ONNX))
+	{
+		return {};
+	}
 
-FString UNNERuntimeORTDmlEditor::GetModelDataIdentifier(const FString& FileType, TConstArrayView<uint8> FileData, const TMap<FString, TConstArrayView<uint8>>& AdditionalFileData, const FGuid& FileId, const ITargetPlatform* TargetPlatform) const
-{
-	return Impl->GetModelDataIdentifier(FileType, FileData, AdditionalFileData, FileId, TargetPlatform);
-}
+	TArray<uint8> Result;
+	FMemoryWriter Writer(Result);
+	Writer << UNNERuntimeORTDml::GUID;
+	Writer << UNNERuntimeORTDml::Version;
+	Writer.Serialize(InputModel.Data.GetData(), InputModel.Data.Num());
 
-UNNERuntimeORTDmlEditor::ECanCreateModelGPUStatus UNNERuntimeORTDmlEditor::CanCreateModelGPU(const TObjectPtr<UNNEModelData> ModelData) const
-{
-	return Impl->CanCreateModelGPU(ModelData);
+	return MakeShared<UE::NNE::FSharedModelData>(MakeSharedBufferFromArray(MoveTemp(Result)), 0);
 }
 
-TSharedPtr<UE::NNE::IModelGPU> UNNERuntimeORTDmlEditor::CreateModelGPU(const TObjectPtr<UNNEModelData> ModelData)
+FString UNNERuntimeORTDml::GetModelDataIdentifier(const FString& FileType, TConstArrayView<uint8> FileData, const TMap<FString, TConstArrayView<uint8>>& AdditionalFileData, const FGuid& FileId, const ITargetPlatform* TargetPlatform) const
 {
-	return Impl->CreateModelGPU(ModelData);
+	return FileId.ToString(EGuidFormats::Digits) + "-" + UNNERuntimeORTDml::GUID.ToString(EGuidFormats::Digits) + "-" + FString::FromInt(UNNERuntimeORTDml::Version);
 }
 
-UNNERuntimeORTDmlEditor::ECanCreateModelRDGStatus UNNERuntimeORTDmlEditor::CanCreateModelRDG(TObjectPtr<UNNEModelData> ModelData) const
+UNNERuntimeORTDml::ECanCreateModelGPUStatus UNNERuntimeORTDml::CanCreateModelGPU(const TObjectPtr<UNNEModelData> ModelData) const
 {
-	return Impl->CanCreateModelRDG(ModelData);
+	return CanCreateModelCommon(ModelData) == ECanCreateModelCommonStatus::Ok ? ECanCreateModelGPUStatus::Ok : ECanCreateModelGPUStatus::Fail;
 }
 
-TSharedPtr<UE::NNE::IModelRDG> UNNERuntimeORTDmlEditor::CreateModelRDG(TObjectPtr<UNNEModelData> ModelData)
+TSharedPtr<UE::NNE::IModelGPU> UNNERuntimeORTDml::CreateModelGPU(const TObjectPtr<UNNEModelData> ModelData)
 {
-	return Impl->CreateModelRDG(ModelData);
-}
+#if PLATFORM_WINDOWS
+	check(ModelData);
 
-/*
- * UNNERuntimeORTDml
- */
-UNNERuntimeORTDml::UNNERuntimeORTDml()
-{
-	Impl = MakeUnique<UE::NNERuntimeORT::Private::NNERuntimeORTDmlImpl>(GUID, Version);
-}
+	if (CanCreateModelGPU(ModelData) != ECanCreateModelGPUStatus::Ok)
+	{
+		UE_LOG(LogNNE, Warning, TEXT("UNNERuntimeORTDml cannot create a model GPU from the model data with id %s"), *ModelData->GetFileId().ToString(EGuidFormats::Digits));
+		return {};
+	}
 
-void UNNERuntimeORTDml::Init()
-{
-	Impl->Init();
-}
+	const TSharedPtr<UE::NNE::FSharedModelData> SharedData = ModelData->GetModelData(GetRuntimeName());
+	check(SharedData.IsValid());
 
-FString UNNERuntimeORTDml::GetRuntimeName() const
-{
-	return Impl->GetRuntimeName();
-}
+	if (FEngineAnalytics::IsAvailable())
+	{
+		TArray<FAnalyticsEventAttribute> Attributes = MakeAnalyticsEventAttributeArray(
+			TEXT("PlatformName"), UGameplayStatics::GetPlatformName(),
+			TEXT("HashedRuntimeName"), FMD5::HashAnsiString(*GetRuntimeName()),
+			TEXT("ModelDataSize"), SharedData->GetView().Num()
+		);
+		FEngineAnalytics::GetProvider().RecordEvent(TEXT("NeuralNetworkEngine.CreateModel"), Attributes);
+	}
 
-UNNERuntimeORTDml::ECanCreateModelDataStatus UNNERuntimeORTDml::CanCreateModelData(const FString& FileType, TConstArrayView<uint8> FileData, const TMap<FString, TConstArrayView<uint8>>& AdditionalFileData, const FGuid& FileId, const ITargetPlatform* TargetPlatform) const
-{
-	return Impl->CanCreateModelData(FileType, FileData, AdditionalFileData, FileId, TargetPlatform);
+	return MakeShared<UE::NNERuntimeORT::Private::FModelORTDmlGPU>(Environment.ToSharedRef(), SharedData);
+#else // PLATFORM_WINDOWS
+	return TSharedPtr<UE::NNE::IModelGPU>();
+#endif // PLATFORM_WINDOWS
 }
 
-TSharedPtr<UE::NNE::FSharedModelData> UNNERuntimeORTDml::CreateModelData(const FString& FileType, TConstArrayView<uint8> FileData, const TMap<FString, TConstArrayView<uint8>>& AdditionalFileData, const FGuid& FileId, const ITargetPlatform* TargetPlatform)
+UNNERuntimeORTDml::ECanCreateModelRDGStatus UNNERuntimeORTDml::CanCreateModelRDG(TObjectPtr<UNNEModelData> ModelData) const
 {
-	return Impl->CreateModelData(FileType, FileData, AdditionalFileData, FileId, TargetPlatform);
+	return CanCreateModelCommon(ModelData) == ECanCreateModelCommonStatus::Ok ? ECanCreateModelRDGStatus::Ok : ECanCreateModelRDGStatus::Fail;
 }
 
-FString UNNERuntimeORTDml::GetModelDataIdentifier(const FString& FileType, TConstArrayView<uint8> FileData, const TMap<FString, TConstArrayView<uint8>>& AdditionalFileData, const FGuid& FileId, const ITargetPlatform* TargetPlatform) const
+TSharedPtr<UE::NNE::IModelRDG> UNNERuntimeORTDml::CreateModelRDG(TObjectPtr<UNNEModelData> ModelData)
 {
-	return Impl->GetModelDataIdentifier(FileType, FileData, AdditionalFileData, FileId, TargetPlatform);
-}
+#if PLATFORM_WINDOWS
+	check(ModelData);
 
-UNNERuntimeORTDml::ECanCreateModelRDGStatus UNNERuntimeORTDml::CanCreateModelRDG(TObjectPtr<UNNEModelData> ModelData) const
-{
-	return Impl->CanCreateModelRDG(ModelData);
+	if (CanCreateModelRDG(ModelData) != ECanCreateModelRDGStatus::Ok)
+	{
+		UE_LOG(LogNNE, Warning, TEXT("UNNERuntimeORTDml cannot create a model RDG from the model data with id %s"), *ModelData->GetFileId().ToString(EGuidFormats::Digits));
+		return {};
+	}
+
+	const TSharedRef<UE::NNE::FSharedModelData> SharedData = ModelData->GetModelData(GetRuntimeName()).ToSharedRef();
+
+	if (FEngineAnalytics::IsAvailable())
+	{
+		TArray<FAnalyticsEventAttribute> Attributes = MakeAnalyticsEventAttributeArray(
+			TEXT("PlatformName"), UGameplayStatics::GetPlatformName(),
+			TEXT("HashedRuntimeName"), FMD5::HashAnsiString(*GetRuntimeName()),
+			TEXT("ModelDataSize"), SharedData->GetView().Num()
+		);
+		FEngineAnalytics::GetProvider().RecordEvent(TEXT("NeuralNetworkEngine.CreateModel"), Attributes);
+	}
+
+	return MakeShared<UE::NNERuntimeORT::Private::FModelORTDmlRDG>(Environment.ToSharedRef(), SharedData);
+#else // PLATFORM_WINDOWS
+	return {};
+#endif // PLATFORM_WINDOWS
 }
 
-TSharedPtr<UE::NNE::IModelRDG> UNNERuntimeORTDml::CreateModelRDG(TObjectPtr<UNNEModelData> ModelData)
+UNNERuntimeORTDml::ECanCreateModelCommonStatus UNNERuntimeORTDml::CanCreateModelCommon(const TObjectPtr<UNNEModelData> ModelData) const
 {
-	return Impl->CreateModelRDG(ModelData);
+#if PLATFORM_WINDOWS
+	check(ModelData != nullptr);
+
+	// In order to use DirectML we need D3D12
+	if (!IsRHID3D12())
+	{
+		return ECanCreateModelCommonStatus::Fail;
+	}
+
+	constexpr int32 GuidSize = sizeof(UNNERuntimeORTDml::GUID);
+	constexpr int32 VersionSize = sizeof(UNNERuntimeORTDml::Version);
+	const TSharedPtr<UE::NNE::FSharedModelData> SharedData = ModelData->GetModelData(GetRuntimeName());
+
+	if (!SharedData.IsValid())
+	{
+		return ECanCreateModelCommonStatus::Fail;
+	}
+
+	TConstArrayView<uint8> Data = SharedData->GetView();
+
+	if (Data.Num() <= GuidSize + VersionSize)
+	{
+		return ECanCreateModelCommonStatus::Fail;
+	}
+
+	static const FGuid DeprecatedGUID = FGuid((int32)'O', (int32)'G', (int32)'P', (int32)'U');
+
+	bool bResult = FGenericPlatformMemory::Memcmp(&(Data[0]), &(UNNERuntimeORTDml::GUID), GuidSize) == 0;
+	bResult |= FGenericPlatformMemory::Memcmp(&(Data[0]), &(DeprecatedGUID), GuidSize) == 0;
+	bResult &= FGenericPlatformMemory::Memcmp(&(Data[GuidSize]), &(UNNERuntimeORTDml::Version), VersionSize) == 0;
+
+	return bResult ? ECanCreateModelCommonStatus::Ok : ECanCreateModelCommonStatus::Fail;
+#else // PLATFORM_WINDOWS
+	return ECanCreateModelCommonStatus::Fail;
+#endif // PLATFORM_WINDOWS
 }
\ No newline at end of file
diff --git a/Engine/Plugins/NNE/NNERuntimeORT/Source/NNERuntimeORT/Private/NNERuntimeORT.h b/Engine/Plugins/NNE/NNERuntimeORT/Source/NNERuntimeORT/Private/NNERuntimeORT.h
--- a/Engine/Plugins/NNE/NNERuntimeORT/Source/NNERuntimeORT/Private/NNERuntimeORT.h
+++ b/Engine/Plugins/NNE/NNERuntimeORT/Source/NNERuntimeORT/Private/NNERuntimeORT.h
@@ -2,13 +2,10 @@
 
 #pragma once
 
-#include "CoreMinimal.h"
-#include "NNEOnnxruntime.h"
 #include "NNERuntime.h"
 #include "NNERuntimeCPU.h"
 #include "NNERuntimeGPU.h"
 #include "NNERuntimeRDG.h"
-#include "Templates/UniquePtr.h"
 #include "UObject/Class.h"
 #include "UObject/Object.h"
 #include "UObject/UObjectBaseUtility.h"
@@ -17,7 +14,7 @@
 
 namespace UE::NNERuntimeORT::Private
 {
-	class NNERuntimeORTDmlImpl;
+	class FEnvironment;
 }
 
 UCLASS()
@@ -26,7 +23,7 @@ class UNNERuntimeORTCpu : public UObject, public INNERuntime, public INNERuntime
 	GENERATED_BODY()
 
 private:
-	TSharedPtr<Ort::Env> ORTEnvironment;
+	TSharedPtr<UE::NNERuntimeORT::Private::FEnvironment> Environment;
 
 public:
 	static FGuid GUID;
@@ -34,8 +31,8 @@ public:
 
 	virtual ~UNNERuntimeORTCpu() = default;
 
-	void Init();
-		
+	void Init(TSharedRef<UE::NNERuntimeORT::Private::FEnvironment> InEnvironment);
+
 	virtual FString GetRuntimeName() const override;
 	
 	virtual ECanCreateModelDataStatus CanCreateModelData(const FString& FileType, TConstArrayView<uint8> FileData, const TMap<FString, TConstArrayView<uint8>>& AdditionalFileData, const FGuid& FileId, const ITargetPlatform* TargetPlatform) const override;
@@ -47,23 +44,22 @@ public:
 };
 
 UCLASS()
-class UNNERuntimeORTDmlEditor : public UObject, public INNERuntime, public INNERuntimeGPU, public INNERuntimeRDG
+class UNNERuntimeORTDml : public UObject, public INNERuntime, public INNERuntimeGPU, public INNERuntimeRDG
 {
 	GENERATED_BODY()
 
 	using ECanCreateModelCommonStatus = UE::NNE::EResultStatus;
 
 private:
-	TUniquePtr<UE::NNERuntimeORT::Private::NNERuntimeORTDmlImpl> Impl;
+	TSharedPtr<UE::NNERuntimeORT::Private::FEnvironment> Environment;
 
 public:
 	static FGuid GUID;
 	static int32 Version;
 	
-	UNNERuntimeORTDmlEditor();
-	virtual ~UNNERuntimeORTDmlEditor() = default;
+	virtual ~UNNERuntimeORTDml() = default;
 
-	void Init();
+	void Init(TSharedRef<UE::NNERuntimeORT::Private::FEnvironment> InEnvironment);
 
 	virtual FString GetRuntimeName() const override;
 
@@ -76,33 +72,7 @@ public:
 
 	virtual ECanCreateModelRDGStatus CanCreateModelRDG(TObjectPtr<UNNEModelData> ModelData) const override;
 	virtual TSharedPtr<UE::NNE::IModelRDG> CreateModelRDG(TObjectPtr<UNNEModelData> ModelData) override;
-};
-
-UCLASS()
-class UNNERuntimeORTDml : public UObject, public INNERuntime, public INNERuntimeRDG
-{
-	GENERATED_BODY()
-
-	using ECanCreateModelCommonStatus = UE::NNE::EResultStatus;
 
 private:
-	TUniquePtr<UE::NNERuntimeORT::Private::NNERuntimeORTDmlImpl> Impl;
-
-public:
-	static FGuid GUID;
-	static int32 Version;
-	
-	UNNERuntimeORTDml();
-	virtual ~UNNERuntimeORTDml() = default;
-
-	void Init();
-
-	virtual FString GetRuntimeName() const override;
-
-	virtual ECanCreateModelDataStatus CanCreateModelData(const FString& FileType, TConstArrayView<uint8> FileData, const TMap<FString, TConstArrayView<uint8>>& AdditionalFileData, const FGuid& FileId, const ITargetPlatform* TargetPlatform) const override;
-	virtual TSharedPtr<UE::NNE::FSharedModelData> CreateModelData(const FString& FileType, TConstArrayView<uint8> FileData, const TMap<FString, TConstArrayView<uint8>>& AdditionalFileData, const FGuid& FileId, const ITargetPlatform* TargetPlatform) override;
-	virtual FString GetModelDataIdentifier(const FString& FileType, TConstArrayView<uint8> FileData, const TMap<FString, TConstArrayView<uint8>>& AdditionalFileData, const FGuid& FileId, const ITargetPlatform* TargetPlatform) const override;
-
-	virtual ECanCreateModelRDGStatus CanCreateModelRDG(TObjectPtr<UNNEModelData> ModelData) const override;
-	virtual TSharedPtr<UE::NNE::IModelRDG> CreateModelRDG(TObjectPtr<UNNEModelData> ModelData) override;
+	ECanCreateModelCommonStatus CanCreateModelCommon(const TObjectPtr<UNNEModelData> ModelData) const;
 };
\ No newline at end of file
diff --git a/Engine/Plugins/NNE/NNERuntimeORT/Source/NNERuntimeORT/Private/NNERuntimeORTEnv.cpp b/Engine/Plugins/NNE/NNERuntimeORT/Source/NNERuntimeORT/Private/NNERuntimeORTEnv.cpp
new file mode 100644
--- /dev/null
+++ b/Engine/Plugins/NNE/NNERuntimeORT/Source/NNERuntimeORT/Private/NNERuntimeORTEnv.cpp
@@ -0,0 +1,45 @@
+// Copyright Epic Games, Inc. All Rights Reserved.
+
+#include "NNERuntimeORTEnv.h"
+
+namespace UE::NNERuntimeORT::Private
+{
+
+void FEnvironment::CreateOrtEnv(const FConfig& InConfig)
+{
+#if WITH_EDITOR
+	// For reloading need to make sure OrtApi::ReleaseEnv() is called before we create a new Ort::Env!
+	OrtEnvironment = Ort::Env(nullptr);
+#endif // WITH_EDITOR
+	checkf(OrtEnvironment == Ort::Env(nullptr), TEXT("Ort::Env already created!"));
+
+	if (Config.bUseGlobalThreadPool)
+	{
+		Ort::ThreadingOptions ThreadingOptions;
+		ThreadingOptions.SetGlobalIntraOpNumThreads(Config.IntraOpNumThreads);
+		ThreadingOptions.SetGlobalInterOpNumThreads(Config.InterOpNumThreads);
+
+		// Calls OrtApi::CreateEnvWithGlobalThreadPools(), needs to be called in conjunction
+		// with OrtApi::DisablePerSessionThreads or else the session will use its own thread pools.
+		OrtEnvironment = Ort::Env(ThreadingOptions);
+	}
+	else
+	{
+		// Calls OrtApi::CreateEnv()
+		OrtEnvironment = Ort::Env();
+	}
+
+	// 
+	// WITH_EDITOR After this point, if we created a new Ort::Env, all Onnxruntime structures (ideally) need to be recreated!
+	// At least SessionOption and Session...
+	// 
+
+	Config = InConfig;
+}
+
+const Ort::Env& FEnvironment::GetOrtEnv() const
+{
+	return OrtEnvironment;
+}
+
+} // UE::NNERuntimeORT::Private
\ No newline at end of file
diff --git a/Engine/Plugins/NNE/NNERuntimeORT/Source/NNERuntimeORT/Private/NNERuntimeORTEnv.h b/Engine/Plugins/NNE/NNERuntimeORT/Source/NNERuntimeORT/Private/NNERuntimeORTEnv.h
new file mode 100644
--- /dev/null
+++ b/Engine/Plugins/NNE/NNERuntimeORT/Source/NNERuntimeORT/Private/NNERuntimeORTEnv.h
@@ -0,0 +1,33 @@
+// Copyright Epic Games, Inc. All Rights Reserved.
+
+#pragma once
+
+#include "NNEOnnxruntime.h"
+
+namespace UE::NNERuntimeORT::Private
+{
+
+class FEnvironment
+{
+public:
+	struct FConfig
+	{
+		bool bUseGlobalThreadPool = false;
+		int32 IntraOpNumThreads = 0;
+		int32 InterOpNumThreads = 0;
+	};
+
+	FEnvironment() = default;
+	
+	~FEnvironment() = default;
+
+	void CreateOrtEnv(const FConfig& InConfig);
+
+	const Ort::Env& GetOrtEnv() const;
+
+private:
+	FConfig Config{};
+	Ort::Env OrtEnvironment{nullptr};
+};
+
+} // UE::NNERuntimeORT::Private
\ No newline at end of file
diff --git a/Engine/Plugins/NNE/NNERuntimeORT/Source/NNERuntimeORT/Private/NNERuntimeORTModel.cpp b/Engine/Plugins/NNE/NNERuntimeORT/Source/NNERuntimeORT/Private/NNERuntimeORTModel.cpp
--- a/Engine/Plugins/NNE/NNERuntimeORT/Source/NNERuntimeORT/Private/NNERuntimeORTModel.cpp
+++ b/Engine/Plugins/NNE/NNERuntimeORT/Source/NNERuntimeORT/Private/NNERuntimeORTModel.cpp
@@ -5,6 +5,7 @@
 #include "HAL/IConsoleManager.h"
 #include "HAL/PlatformMisc.h"
 #include "NNERuntimeORT.h"
+#include "NNERuntimeORTSettings.h"
 #include "NNERuntimeORTUtils.h"
 #include "RenderGraphUtils.h"
 
@@ -40,11 +41,46 @@ namespace UE::NNERuntimeORT::Private
 namespace Detail
 {
 
+	FRuntimeConf MakeRuntimeConfigFromSettings(const UNNERuntimeORTSettings* Settings)
+	{
+		FRuntimeConf Result{};
+
+#if WITH_EDITOR
+			FThreadingOptions ThreadingOptions = Settings->EditorThreadingOptions;
+#else
+			FThreadingOptions ThreadingOptions = Settings->GameThreadingOptions;
+#endif
+
+		Result.bUseGlobalThreadPool = ThreadingOptions.bUseGlobalThreadPool;
+		Result.IntraOpNumThreads = ThreadingOptions.IntraOpNumThreads;
+		Result.InterOpNumThreads = ThreadingOptions.InterOpNumThreads;
+		Result.ExecutionMode = ThreadingOptions.ExecutionMode == EExecutionMode::SEQUENTIAL ? ExecutionMode::ORT_SEQUENTIAL : ExecutionMode::ORT_PARALLEL;
+
+		return Result;
+	}
+
 TUniquePtr<Ort::SessionOptions> CreateSessionOptionsDefault(const FRuntimeConf &RuntimeConf)
 {
 	TUniquePtr<Ort::SessionOptions> SessionOptions = MakeUnique<Ort::SessionOptions>();
+
+	// Configure Threading
+	if (RuntimeConf.bUseGlobalThreadPool)
+	{
+		SessionOptions->DisablePerSessionThreads();
+	}
+	else
+	{
+		SessionOptions->SetIntraOpNumThreads(RuntimeConf.IntraOpNumThreads);
+		SessionOptions->SetInterOpNumThreads(RuntimeConf.InterOpNumThreads);
+	}
+
+	SessionOptions->SetExecutionMode(RuntimeConf.ExecutionMode);
+
+	// Configure Graph optimization
 	SessionOptions->SetGraphOptimizationLevel(RuntimeConf.OptimizationLevel);
 
+
+	// Configure Profiling
 	if (CVarNNERuntimeORTEnableProfiling.GetValueOnGameThread())
 	{
 		FString ProfilingFilePrefix("NNERuntimeORTProfile_");
@@ -64,6 +100,12 @@ TUniquePtr<Ort::SessionOptions> CreateSessionOptionsDefault(const FRuntimeConf &
 TUniquePtr<Ort::SessionOptions> CreateSessionOptionsForDirectML(const FRuntimeConf &RuntimeConf)
 {
 	TUniquePtr<Ort::SessionOptions> SessionOptions = CreateSessionOptionsDefault(RuntimeConf);
+	if (!SessionOptions.IsValid())
+	{
+		return {};
+	}
+
+	// Configure for DirectML
 	SessionOptions->SetExecutionMode(ExecutionMode::ORT_SEQUENTIAL);
 	SessionOptions->DisableMemPattern();
 
@@ -148,7 +190,7 @@ TUniquePtr<Ort::SessionOptions> CreateSessionOptionsForDirectML(const FRuntimeCo
 } // namespace Detail
 
 template <class ModelInterface, class TensorBinding> 
-FModelInstanceORTBase<ModelInterface, TensorBinding>::FModelInstanceORTBase(const FRuntimeConf& InRuntimeConf, TSharedPtr<Ort::Env> InEnvironment)
+FModelInstanceORTBase<ModelInterface, TensorBinding>::FModelInstanceORTBase(const FRuntimeConf& InRuntimeConf, TSharedRef<FEnvironment> InEnvironment)
 	: RuntimeConf(InRuntimeConf), Environment(InEnvironment)
 {
 
@@ -177,7 +219,7 @@ bool FModelInstanceORTBase<ModelInterface, TensorBinding>::Init(TConstArrayView<
 			return false;
 		}
 
-		Session = MakeUnique<Ort::Session>(*Environment, ModelBuffer.GetData(), ModelBuffer.Num(), *SessionOptions);
+		Session = MakeUnique<Ort::Session>(Environment->GetOrtEnv(), ModelBuffer.GetData(), ModelBuffer.Num(), *SessionOptions);
 
 		if (!ConfigureTensors(true))
 		{
@@ -368,10 +410,9 @@ typename ModelInterface::ERunSyncStatus FModelInstanceORTBase<ModelInterface, Te
 	return ModelInterface::ERunSyncStatus::Ok;
 }
 
-#if WITH_EDITOR
 TSharedPtr<NNE::IModelInstanceCPU> FModelORTCpu::CreateModelInstanceCPU()
 {
-	const FRuntimeConf Conf;
+	const FRuntimeConf Conf = Detail::MakeRuntimeConfigFromSettings(GetDefault<UNNERuntimeORTSettings>());
 	FModelInstanceORTCpu* ModelInstance = new FModelInstanceORTCpu(Conf, Environment);
 
 	check(ModelData.IsValid());
@@ -385,7 +426,7 @@ TSharedPtr<NNE::IModelInstanceCPU> FModelORTCpu::CreateModelInstanceCPU()
 	return TSharedPtr<NNE::IModelInstanceCPU>(IModelInstance);
 }
 
-FModelORTCpu::FModelORTCpu(TSharedPtr<Ort::Env> InEnvironment, const TSharedPtr<UE::NNE::FSharedModelData>& InModelData) :
+FModelORTCpu::FModelORTCpu(TSharedRef<FEnvironment> InEnvironment, const TSharedPtr<UE::NNE::FSharedModelData>& InModelData) :
 	Environment(InEnvironment), ModelData(InModelData)
 {
 }
@@ -407,12 +448,11 @@ bool FModelInstanceORTCpu::InitializedAndConfigureMembers()
 
 	return true;
 }
-#endif // WITH_EDITOR
 
 #if PLATFORM_WINDOWS
 TSharedPtr<NNE::IModelInstanceGPU> FModelORTDmlGPU::CreateModelInstanceGPU()
 {
-	const FRuntimeConf Conf;
+	const FRuntimeConf Conf = Detail::MakeRuntimeConfigFromSettings(GetDefault<UNNERuntimeORTSettings>());
 	FModelInstanceORTDmlGPU* ModelInstance = new FModelInstanceORTDmlGPU(Conf, Environment);
 
 	check(ModelData.IsValid());
@@ -426,7 +466,7 @@ TSharedPtr<NNE::IModelInstanceGPU> FModelORTDmlGPU::CreateModelInstanceGPU()
 	return TSharedPtr<NNE::IModelInstanceGPU>(IModelInstance);
 }
 
-FModelORTDmlGPU::FModelORTDmlGPU(TSharedPtr<Ort::Env> InEnvironment, const TSharedPtr<UE::NNE::FSharedModelData>& InModelData) :
+FModelORTDmlGPU::FModelORTDmlGPU(TSharedRef<FEnvironment> InEnvironment, const TSharedPtr<UE::NNE::FSharedModelData>& InModelData) :
 	Environment(InEnvironment), ModelData(InModelData)
 {
 }
@@ -447,7 +487,7 @@ bool FModelInstanceORTDmlGPU::InitializedAndConfigureMembers()
 	return true;
 }
 
-FModelORTDmlRDG::FModelORTDmlRDG(TSharedRef<Ort::Env> InEnvironment, TSharedRef<UE::NNE::FSharedModelData> InModelData) :
+FModelORTDmlRDG::FModelORTDmlRDG(TSharedRef<FEnvironment> InEnvironment, TSharedRef<UE::NNE::FSharedModelData> InModelData) :
 	Environment(InEnvironment), ModelData(InModelData)
 {
 }
@@ -463,7 +503,7 @@ TSharedPtr<NNE::IModelInstanceRDG> FModelORTDmlRDG::CreateModelInstanceRDG()
 	return ModelInstance;
 }
 
-FModelInstanceORTDmlRDG::FModelInstanceORTDmlRDG(TSharedRef<UE::NNE::FSharedModelData> InModelData, const FRuntimeConf& InRuntimeConf, TSharedRef<Ort::Env> InEnvironment)
+FModelInstanceORTDmlRDG::FModelInstanceORTDmlRDG(TSharedRef<UE::NNE::FSharedModelData> InModelData, const FRuntimeConf& InRuntimeConf, TSharedRef<FEnvironment> InEnvironment)
 	:  ModelData(InModelData), RuntimeConf(InRuntimeConf), Environment(InEnvironment)
 {}
 
@@ -488,11 +528,11 @@ bool FModelInstanceORTDmlRDG::Init()
 		SessionOptions = Detail::CreateSessionOptionsForDirectML(RuntimeConf);
 		if (!SessionOptions.IsValid())
 		{
-			UE_LOG(LogNNE, Error, TEXT("FModelInstanceORTDmlRDG::Init(): Failed to create session options."));
+			UE_LOG(LogNNE, Error, TEXT("FModelInstanceORTDmlRDG::Init(): Failed to configure session options for DirectML Execution Provider."));
 			return false;
 		}
 
-		Session = MakeUnique<Ort::Session>(*Environment, ModelBuffer.GetData(), ModelBuffer.Num(), *SessionOptions);
+		Session = MakeUnique<Ort::Session>(Environment->GetOrtEnv(), ModelBuffer.GetData(), ModelBuffer.Num(), *SessionOptions);
 
 		if (!ConfigureTensors(*Session))
 		{
@@ -681,7 +721,7 @@ FModelInstanceORTDmlRDG::ESetInputTensorShapesStatus FModelInstanceORTDmlRDG::Se
 	TConstArrayView<uint8> ModelBuffer = TConstArrayView<uint8>(&(ModelData->GetView().GetData()[GuidSize + VersionSize]), ModelData->GetView().Num() - GuidSize - VersionSize);
 	check(!ModelBuffer.IsEmpty());
 
-	Session = MakeUnique<Ort::Session>(*Environment, ModelBuffer.GetData(), ModelBuffer.Num(), *SessionOptions);
+	Session = MakeUnique<Ort::Session>(Environment->GetOrtEnv(), ModelBuffer.GetData(), ModelBuffer.Num(), *SessionOptions);
 
 	// Need to configure output tensors with new session (to apply free dimension overrides)
 	if (!ConfigureTensors(*Session, false))
diff --git a/Engine/Plugins/NNE/NNERuntimeORT/Source/NNERuntimeORT/Private/NNERuntimeORTModel.h b/Engine/Plugins/NNE/NNERuntimeORT/Source/NNERuntimeORT/Private/NNERuntimeORTModel.h
--- a/Engine/Plugins/NNE/NNERuntimeORT/Source/NNERuntimeORT/Private/NNERuntimeORTModel.h
+++ b/Engine/Plugins/NNE/NNERuntimeORT/Source/NNERuntimeORT/Private/NNERuntimeORTModel.h
@@ -8,6 +8,7 @@
 #include "NNEOnnxruntime.h"
 #include "NNERuntimeCPU.h"
 #include "NNERuntimeGPU.h"
+#include "NNERuntimeORTEnv.h"
 #include "NNERuntimeRDG.h"
 #include "NNETensor.h"
 #include "NNETypes.h"
@@ -17,6 +18,12 @@ namespace UE::NNERuntimeORT::Private
 
 struct FRuntimeConf
 {
+	bool bUseGlobalThreadPool = false;
+	int32 IntraOpNumThreads = 0;
+	int32 InterOpNumThreads = 0;
+
+	ExecutionMode ExecutionMode = ExecutionMode::ORT_SEQUENTIAL;
+
 	GraphOptimizationLevel OptimizationLevel = GraphOptimizationLevel::ORT_ENABLE_ALL;
 };
 
@@ -25,7 +32,7 @@ class FModelInstanceORTBase : public NNE::Internal::FModelInstanceBase<ModelInte
 {
 
 public:
-	FModelInstanceORTBase(const FRuntimeConf& InRuntimeConf, TSharedPtr<Ort::Env> InEnvironment);
+	FModelInstanceORTBase(const FRuntimeConf& InRuntimeConf, TSharedRef<FEnvironment> InEnvironment);
 	virtual ~FModelInstanceORTBase() = default;
 
 	virtual typename ModelInterface::ESetInputTensorShapesStatus SetInputTensorShapes(TConstArrayView<NNE::FTensorShape> InInputShapes) override;
@@ -41,7 +48,7 @@ protected:
 	FRuntimeConf RuntimeConf;
 
 	/** ORT-related variables */
-	TSharedPtr<Ort::Env> Environment;
+	TSharedRef<FEnvironment> Environment;
 	TUniquePtr<Ort::Session> Session;
 	TUniquePtr<Ort::AllocatorWithDefaultOptions> Allocator;
 	TUniquePtr<Ort::SessionOptions> SessionOptions;
@@ -60,11 +67,10 @@ protected:
 	TArray<NNE::Internal::FTensor> OutputTensors;
 };
 
-#if WITH_EDITOR
 class FModelInstanceORTCpu : public FModelInstanceORTBase<NNE::IModelInstanceCPU, NNE::FTensorBindingCPU>
 {
 public:
-	FModelInstanceORTCpu(const FRuntimeConf& InRuntimeConf, TSharedPtr<Ort::Env> InEnvironment) : FModelInstanceORTBase(InRuntimeConf, InEnvironment) {}
+	FModelInstanceORTCpu(const FRuntimeConf& InRuntimeConf, TSharedRef<FEnvironment> InEnvironment) : FModelInstanceORTBase(InRuntimeConf, InEnvironment) {}
 	virtual ~FModelInstanceORTCpu() = default;
 
 private:
@@ -74,35 +80,34 @@ private:
 class FModelORTCpu : public NNE::IModelCPU
 {
 public:
-	FModelORTCpu(TSharedPtr<Ort::Env> InEnvironment, const TSharedPtr<UE::NNE::FSharedModelData>& InModelData);
+	FModelORTCpu(TSharedRef<FEnvironment> InEnvironment, const TSharedPtr<UE::NNE::FSharedModelData>& InModelData);
 	virtual ~FModelORTCpu() = default;
 
 	virtual TSharedPtr<NNE::IModelInstanceCPU> CreateModelInstanceCPU() override;
 
 private:
-	TSharedPtr<Ort::Env> Environment;
+	TSharedRef<FEnvironment> Environment;
 	TSharedPtr<UE::NNE::FSharedModelData> ModelData;
 };
-#endif // WITH_EDITOR
 
 #if PLATFORM_WINDOWS
 class FModelORTDmlGPU : public NNE::IModelGPU
 {
 public:
-	FModelORTDmlGPU(TSharedPtr<Ort::Env> InEnvironment, const TSharedPtr<UE::NNE::FSharedModelData>& InModelData);
+	FModelORTDmlGPU(TSharedRef<FEnvironment> InEnvironment, const TSharedPtr<UE::NNE::FSharedModelData>& InModelData);
 	virtual ~FModelORTDmlGPU() = default;
 
 	virtual TSharedPtr<NNE::IModelInstanceGPU> CreateModelInstanceGPU() override;
 
 private:
-	TSharedPtr<Ort::Env> Environment;
+	TSharedRef<FEnvironment> Environment;
 	TSharedPtr<UE::NNE::FSharedModelData> ModelData;
 };
 
 class FModelInstanceORTDmlGPU : public FModelInstanceORTBase<NNE::IModelInstanceGPU, NNE::FTensorBindingGPU>
 {
 public:
-	FModelInstanceORTDmlGPU(const FRuntimeConf& InRuntimeConf, TSharedPtr<Ort::Env> InEnvironment) : FModelInstanceORTBase(InRuntimeConf, InEnvironment) {}
+	FModelInstanceORTDmlGPU(const FRuntimeConf& InRuntimeConf, TSharedRef<FEnvironment> InEnvironment) : FModelInstanceORTBase(InRuntimeConf, InEnvironment) {}
 	virtual ~FModelInstanceORTDmlGPU() = default;
 
 private:
@@ -112,14 +117,14 @@ private:
 class FModelORTDmlRDG : public NNE::IModelRDG
 {
 public:
-	FModelORTDmlRDG(TSharedRef<Ort::Env> InEnvironment, TSharedRef<UE::NNE::FSharedModelData> InModelData);
+	FModelORTDmlRDG(TSharedRef<FEnvironment> InEnvironment, TSharedRef<UE::NNE::FSharedModelData> InModelData);
 
 	virtual ~FModelORTDmlRDG() = default;
 
 	virtual TSharedPtr<NNE::IModelInstanceRDG> CreateModelInstanceRDG() override;
 
 private:
-	TSharedRef<Ort::Env> Environment;
+	TSharedRef<FEnvironment> Environment;
 	TSharedRef<UE::NNE::FSharedModelData> ModelData;
 };
 
@@ -130,7 +135,7 @@ public:
 	using ESetInputTensorShapesStatus = NNE::IModelInstanceRDG::ESetInputTensorShapesStatus;
 	using EEnqueueRDGStatus = NNE::IModelInstanceRDG::EEnqueueRDGStatus;
 
-	FModelInstanceORTDmlRDG(TSharedRef<UE::NNE::FSharedModelData> InModelData, const FRuntimeConf& InRuntimeConf, TSharedRef<Ort::Env> InEnvironment);
+	FModelInstanceORTDmlRDG(TSharedRef<UE::NNE::FSharedModelData> InModelData, const FRuntimeConf& InRuntimeConf, TSharedRef<FEnvironment> InEnvironment);
 
 	virtual ~FModelInstanceORTDmlRDG() = default;
 
@@ -148,7 +153,7 @@ protected:
 	FRuntimeConf RuntimeConf;
 
 	/** ORT-related variables */
-	TSharedRef<Ort::Env> Environment;
+	TSharedRef<FEnvironment> Environment;
 	TUniquePtr<Ort::Session> Session;
 	TUniquePtr<Ort::SessionOptions> SessionOptions;
 	TUniquePtr<Ort::AllocatorWithDefaultOptions> Allocator;
diff --git a/Engine/Plugins/NNE/NNERuntimeORT/Source/NNERuntimeORT/Private/NNERuntimeORTModule.cpp b/Engine/Plugins/NNE/NNERuntimeORT/Source/NNERuntimeORT/Private/NNERuntimeORTModule.cpp
--- a/Engine/Plugins/NNE/NNERuntimeORT/Source/NNERuntimeORT/Private/NNERuntimeORTModule.cpp
+++ b/Engine/Plugins/NNE/NNERuntimeORT/Source/NNERuntimeORT/Private/NNERuntimeORTModule.cpp
@@ -8,40 +8,65 @@
 #include "NNE.h"
 #include "NNEOnnxruntime.h"
 #include "NNERuntimeORT.h"
+#include "NNERuntimeORTEnv.h"
+#include "NNERuntimeORTSettings.h"
 #include "UObject/WeakInterfacePtr.h"
 
-namespace UE::NNERuntimeORT::Private::DllHelper
+namespace UE::NNERuntimeORT::Private
 {
-	bool GetDllHandle(const FString& DllPath, TArray<void*>& DllHandles)
+	namespace DllHelper
 	{
-		void *DllHandle = nullptr;
-
-		if (!FPaths::FileExists(DllPath))
+		bool GetDllHandle(const FString& DllPath, TArray<void*>& DllHandles)
 		{
-			UE_LOG(LogNNE, Error, TEXT("Failed to find the third party library %s."), *DllPath);
-			return false;
+			void *DllHandle = nullptr;
+
+			if (!FPaths::FileExists(DllPath))
+			{
+				UE_LOG(LogNNE, Error, TEXT("Failed to find the third party library %s."), *DllPath);
+				return false;
+			}
+			
+			DllHandle = FPlatformProcess::GetDllHandle(*DllPath);
+
+			if (!DllHandle)
+			{
+				UE_LOG(LogNNE, Error, TEXT("Failed to load the third party library %s."), *DllPath);
+				return false;
+			}
+
+			DllHandles.Add(DllHandle);
+			return true;
 		}
-		
-		DllHandle = FPlatformProcess::GetDllHandle(*DllPath);
+	} // namespace DllHelper
 
-		if (!DllHandle)
+	namespace EnvironmentHelper
+	{
+		void CreateOrtEnvFromSettings(const UNNERuntimeORTSettings* Settings, FEnvironment& Environment)
 		{
-			UE_LOG(LogNNE, Error, TEXT("Failed to load the third party library %s."), *DllPath);
-			return false;
-		}
+#if WITH_EDITOR
+			FThreadingOptions ThreadingOptions = Settings->EditorThreadingOptions;
+#else
+			FThreadingOptions ThreadingOptions = Settings->GameThreadingOptions;
+#endif
 
-		DllHandles.Add(DllHandle);
-		return true;
-	}
+			UE::NNERuntimeORT::Private::FEnvironment::FConfig Config{};
+			Config.bUseGlobalThreadPool = ThreadingOptions.bUseGlobalThreadPool;
+			Config.IntraOpNumThreads = ThreadingOptions.IntraOpNumThreads;
+			Config.InterOpNumThreads = ThreadingOptions.InterOpNumThreads;
 
-} // namespace UE::NNERuntimeRDG::Private::Dml
+			Environment.CreateOrtEnv(Config);
+		}
+	} // EnvironmentHelper
+} // namespace UE::NNERuntimeRDG::Private
 
 void FNNERuntimeORTModule::StartupModule()
 {
+	using namespace UE::NNERuntimeORT::Private;
+
 	const FString PluginDir = IPluginManager::Get().FindPlugin("NNERuntimeORT")->GetBaseDir();
 	const FString OrtSharedLibPath = FPaths::Combine(PluginDir, TEXT(PREPROCESSOR_TO_STRING(ONNXRUNTIME_SHAREDLIB_PATH)));
 
-	if (!UE::NNERuntimeORT::Private::DllHelper::GetDllHandle(OrtSharedLibPath, DllHandles))
+	if (!DllHelper::GetDllHandle(OrtSharedLibPath, DllHandles))
 	{
 		UE_LOG(LogNNE, Error, TEXT("Failed to load ONNX Runtime shared library. ORT Runtimes won't be available."));
 		return;
@@ -53,7 +78,7 @@ void FNNERuntimeORTModule::StartupModule()
 	const FString ModuleDir = FPlatformProcess::GetModulesDirectory();
 	const FString DirectMLSharedLibPath = FPaths::Combine(ModuleDir, TEXT(PREPROCESSOR_TO_STRING(DIRECTML_PATH)), TEXT("DirectML.dll"));
 
-	const bool bDirectMLDllLoaded = UE::NNERuntimeORT::Private::DllHelper::GetDllHandle(DirectMLSharedLibPath, DllHandles);
+	const bool bDirectMLDllLoaded = DllHelper::GetDllHandle(DirectMLSharedLibPath, DllHandles);
 	if (!bDirectMLDllLoaded)
 	{
 		UE_LOG(LogNNE, Error, TEXT("Failed to load DirectML shared library. ORT Dml Runtime won't be available."));
@@ -68,53 +93,52 @@ void FNNERuntimeORTModule::StartupModule()
 	}
 
 	Ort::InitApi(OrtApiFunctions->OrtGetApiBase()->GetApi(ORT_API_VERSION));
+	
+	Environment = MakeShared<FEnvironment>();
+
+	EnvironmentHelper::CreateOrtEnvFromSettings(GetDefault<UNNERuntimeORTSettings>(), *Environment);
 
 #if PLATFORM_WINDOWS
 	if (bDirectMLDllLoaded)
 	{
-#if WITH_EDITOR
-		// NNE runtime ORT Dml startup
-		NNERuntimeORTDmlEditor = NewObject<UNNERuntimeORTDmlEditor>();
-		if (NNERuntimeORTDmlEditor.IsValid())
-		{
-			TWeakInterfacePtr<INNERuntime> RuntimeDmlInterface(NNERuntimeORTDmlEditor.Get());
-
-			NNERuntimeORTDmlEditor->Init();
-			NNERuntimeORTDmlEditor->AddToRoot();
-			UE::NNE::RegisterRuntime(RuntimeDmlInterface);
-		}
-#else
 		// NNE runtime ORT Dml startup
 		NNERuntimeORTDml = NewObject<UNNERuntimeORTDml>();
 		if (NNERuntimeORTDml.IsValid())
 		{
 			TWeakInterfacePtr<INNERuntime> RuntimeDmlInterface(NNERuntimeORTDml.Get());
 
-			NNERuntimeORTDml->Init();
+			NNERuntimeORTDml->Init(Environment.ToSharedRef());
 			NNERuntimeORTDml->AddToRoot();
 			UE::NNE::RegisterRuntime(RuntimeDmlInterface);
 		}
-#endif // WITH_EDITOR
 	}
 #endif // PLATFORM_WINDOWS
 
-#if WITH_EDITOR
 	// NNE runtime ORT Cpu startup
 	NNERuntimeORTCpu = NewObject<UNNERuntimeORTCpu>();
 	if (NNERuntimeORTCpu.IsValid())
 	{
 		TWeakInterfacePtr<INNERuntime> RuntimeCPUInterface(NNERuntimeORTCpu.Get());
 
-		NNERuntimeORTCpu->Init();
+		NNERuntimeORTCpu->Init(Environment.ToSharedRef());
 		NNERuntimeORTCpu->AddToRoot();
 		UE::NNE::RegisterRuntime(RuntimeCPUInterface);
 	}
-#endif // WITH_EDITOR
+
+#if WITH_EDITOR
+	GetMutableDefault<UNNERuntimeORTSettings>()->OnSettingChanged().AddRaw(this, &FNNERuntimeORTModule::OnSettingsChanged);
+#endif
 }
 
 void FNNERuntimeORTModule::ShutdownModule()
 {
 #if WITH_EDITOR
+	if(UObjectInitialized())
+	{
+		GetMutableDefault<UNNERuntimeORTSettings>()->OnSettingChanged().RemoveAll(this);
+	}
+#endif
+
 	// NNE runtime ORT Cpu shutdown
 	if (NNERuntimeORTCpu.IsValid())
 	{
@@ -125,16 +149,6 @@ void FNNERuntimeORTModule::ShutdownModule()
 		NNERuntimeORTCpu.Reset();
 	}
 
-	// NNE runtime ORT Dml shutdown
-	if (NNERuntimeORTDmlEditor.IsValid())
-	{
-		TWeakInterfacePtr<INNERuntime> RuntimeDmlInterface(NNERuntimeORTDmlEditor.Get());
-
-		UE::NNE::UnregisterRuntime(RuntimeDmlInterface);
-		NNERuntimeORTDmlEditor->RemoveFromRoot();
-		NNERuntimeORTDmlEditor.Reset();
-	}
-#else
 	// NNE runtime ORT Dml shutdown
 	if (NNERuntimeORTDml.IsValid())
 	{
@@ -144,7 +158,8 @@ void FNNERuntimeORTModule::ShutdownModule()
 		NNERuntimeORTDml->RemoveFromRoot();
 		NNERuntimeORTDml.Reset();
 	}
-#endif // WITH_EDITOR
+
+	Environment.Reset();
 
 	// Free the dll handles
 	for(void* DllHandle : DllHandles)
@@ -154,4 +169,14 @@ void FNNERuntimeORTModule::ShutdownModule()
 	DllHandles.Empty();
 }
 
+#if WITH_EDITOR
+void FNNERuntimeORTModule::OnSettingsChanged(UObject* InObject, struct FPropertyChangedEvent& InPropertyChangedEvent)
+{
+	UE_LOG(LogNNE, Log, TEXT("Settings %s changed: %s"), *InObject->GetName(), *InPropertyChangedEvent.GetPropertyName().ToString());
+	UE_LOG(LogNNE, Warning, TEXT("It is recommended to restart the Editor if settings %s changed! Otherwise they might not be fully applied."), *InObject->GetName());
+
+	UE::NNERuntimeORT::Private::EnvironmentHelper::CreateOrtEnvFromSettings(CastChecked<UNNERuntimeORTSettings>(InObject), *Environment);
+}
+#endif
+
 IMPLEMENT_MODULE(FNNERuntimeORTModule, NNERuntimeORT);
\ No newline at end of file
diff --git a/Engine/Plugins/NNE/NNERuntimeORT/Source/NNERuntimeORT/Private/NNERuntimeORTModule.h b/Engine/Plugins/NNE/NNERuntimeORT/Source/NNERuntimeORT/Private/NNERuntimeORTModule.h
--- a/Engine/Plugins/NNE/NNERuntimeORT/Source/NNERuntimeORT/Private/NNERuntimeORTModule.h
+++ b/Engine/Plugins/NNE/NNERuntimeORT/Source/NNERuntimeORT/Private/NNERuntimeORTModule.h
@@ -8,19 +8,27 @@
 #include "UObject/WeakObjectPtr.h"
 #include "UObject/WeakObjectPtrTemplates.h"
 
-class UNNERuntimeORTDmlEditor;
 class UNNERuntimeORTDml;
 class UNNERuntimeORTCpu;
 
+namespace UE::NNERuntimeORT::Private
+{
+	class FEnvironment;
+}
+
 class FNNERuntimeORTModule : public IModuleInterface
 {
 private:
-	TWeakObjectPtr<UNNERuntimeORTDmlEditor> NNERuntimeORTDmlEditor{ nullptr };
 	TWeakObjectPtr<UNNERuntimeORTDml> NNERuntimeORTDml{ nullptr };
 	TWeakObjectPtr<UNNERuntimeORTCpu> NNERuntimeORTCpu{ nullptr };
 
 	virtual void StartupModule() override;
 	virtual void ShutdownModule() override;
 
+#if WITH_EDITOR
+void OnSettingsChanged(UObject* InObject, struct FPropertyChangedEvent& InPropertyChangedEvent);
+#endif
+
 	TArray<void*> DllHandles;
+	TSharedPtr<UE::NNERuntimeORT::Private::FEnvironment> Environment;
 };
\ No newline at end of file
diff --git a/Engine/Plugins/NNE/NNERuntimeORT/Source/NNERuntimeORT/Private/NNERuntimeORTSettings.h b/Engine/Plugins/NNE/NNERuntimeORT/Source/NNERuntimeORT/Private/NNERuntimeORTSettings.h
new file mode 100644
--- /dev/null
+++ b/Engine/Plugins/NNE/NNERuntimeORT/Source/NNERuntimeORT/Private/NNERuntimeORTSettings.h
@@ -0,0 +1,64 @@
+// Copyright Epic Games, Inc. All Rights Reserved.
+
+#pragma once
+
+#include "Engine/DeveloperSettings.h"
+
+#include "NNERuntimeORTSettings.generated.h"
+
+/** An enum to represent denoiser NNE runtime type */
+UENUM()
+enum EExecutionMode : uint8
+{
+	SEQUENTIAL,
+	PARALLEL
+};
+
+/** Rigid body error correction data */
+USTRUCT()
+struct FThreadingOptions
+{
+	GENERATED_BODY()
+
+	/** Use global threadpools that will be shared across sessions */
+	UPROPERTY(Config, EditAnywhere, Category = "ONNX Runtime", meta = (
+		DisplayName = "Use global thread pool",
+		ToolTip = "Use global threadpools that will be shared across sessions."))
+	bool bUseGlobalThreadPool = false;
+
+	/** Thread count to parallelize the execution within nodes */
+	UPROPERTY(Config, EditAnywhere, Category = "ONNX Runtime", meta = (
+		DisplayName = "Intra-op thread count",
+		ToolTip = 	"Set thread count of intra-op thread pool, which is utilized by ONNX Runtime to parallelize computation inside each operator.\nSpecial values:\n   0 = Use default thread count\n   1 = The invoking thread will be used; no threads will be created in the thread pool."))
+	int32 IntraOpNumThreads = 0;
+
+	/** Thread count used to parallelize the execution of the graph */
+	UPROPERTY(Config, EditAnywhere, Category = "ONNX Runtime", meta = (
+		DisplayName = "Inter-op thread count",
+		ToolTip = 	"Set thread count of the inter-op thread pool, which enables parallelism between operators and will only be created when session execution mode set to parallel.\nSpecial values:\n   0 = Use default thread count\n   1 = The invoking thread will be used; no threads will be created in the thread pool."))
+	int32 InterOpNumThreads = 0;
+
+	/** Execution mode controls whether multiple operators in the graph (across nodes) run sequentially or in parallel */
+	UPROPERTY(Config, EditAnywhere, Category = "ONNX Runtime", meta = (
+		DisplayName = "Execution mode",
+		ToolTip = 	"Controls whether multiple operators in the graph (across nodes) run sequentially or in parallel.\nNote: DirectML EP requires sequential execution and therefore ignores this setting."))
+	TEnumAsByte<EExecutionMode> ExecutionMode = EExecutionMode::SEQUENTIAL;
+};
+
+/** Settings used to configure NNERuntimeORT */
+UCLASS(Config = Engine, meta = (DisplayName = "NNERuntimeORT"))
+class NNERUNTIMEORT_API UNNERuntimeORTSettings : public UDeveloperSettings
+{
+	GENERATED_BODY()
+
+public:
+
+	/** Threading options in Editor targets */
+	UPROPERTY(Config, EditAnywhere, Category = "ONNX Runtime", meta = (DisplayName = "Editor threading options", ToolTip = "Threading options in Editor targets", ConfigRestartRequired = true))
+	FThreadingOptions EditorThreadingOptions{false, 0, 0, EExecutionMode::SEQUENTIAL};
+
+	/** Threading options in Non-Editor (Game, Program, ...) targets */
+	UPROPERTY(Config, EditAnywhere, Category = "ONNX Runtime", meta = (DisplayName = "Game threading options", ToolTip = "Threading options in Non-Editor (Game, Program, ...) targets", ConfigRestartRequired = true))
+	FThreadingOptions GameThreadingOptions{false, 1, 1, EExecutionMode::SEQUENTIAL};
+
+};
