[NNE] NNERuntimeORT Standardize ModelInstance API behaviour.

#tests Editor/Game Win64 all tests

#rnx

diff --git a/Engine/Plugins/NNE/NNERuntimeORT/Source/NNERuntimeORT/Private/NNERuntimeORTModel.cpp b/Engine/Plugins/NNE/NNERuntimeORT/Source/NNERuntimeORT/Private/NNERuntimeORTModel.cpp
--- a/Engine/Plugins/NNE/NNERuntimeORT/Source/NNERuntimeORT/Private/NNERuntimeORTModel.cpp
+++ b/Engine/Plugins/NNE/NNERuntimeORT/Source/NNERuntimeORT/Private/NNERuntimeORTModel.cpp
@@ -252,7 +252,7 @@ template <class ModelInterface, class TensorBinding>
 bool FModelInstanceORTBase<ModelInterface, TensorBinding>::InitializedAndConfigureMembers()
 {
 	Allocator = MakeUnique<Ort::AllocatorWithDefaultOptions>();
-	AllocatorInfo = MakeUnique<Ort::MemoryInfo>(Ort::MemoryInfo::CreateCpu(OrtDeviceAllocator, OrtMemTypeCPU));
+	MemoryInfo = MakeUnique<Ort::MemoryInfo>(Ort::MemoryInfo::CreateCpu(OrtDeviceAllocator, OrtMemTypeCPU));
 
 	return true;
 }
@@ -342,17 +342,54 @@ typename ModelInterface::ESetInputTensorShapesStatus FModelInstanceORTBase<Model
 	return ModelInterface::ESetInputTensorShapesStatus::Ok;
 }
 
+template <class TensorBinding>
+Ort::Value CreateTensor(const Ort::MemoryInfo& MemoryInfo, const TensorBinding& Binding, const NNE::Internal::FTensor& Tensor, const ONNXTensorElementDataType ElementDataType)
+{
+	const uint64 SizeInBytes = Tensor.GetDataSize();
+	const uint32 ShapeLen = (uint32)Tensor.GetShape().Rank();
+
+	TUniquePtr<int64_t[]> Shape = MakeUnique<int64_t[]>(Tensor.GetShape().Rank());
+	for (int32 DimIndex = 0; DimIndex < Tensor.GetShape().Rank(); ++DimIndex)
+	{
+		Shape.Get()[DimIndex] = Tensor.GetShape().GetData()[DimIndex];
+	}
+	
+	return	Ort::Value::CreateTensor(MemoryInfo, Binding.Data, SizeInBytes, Shape.Get(), ShapeLen, ElementDataType);
+}
+
 template <class ModelInterface, class TensorBinding>
 typename ModelInterface::ERunSyncStatus FModelInstanceORTBase<ModelInterface, TensorBinding>::RunSync(TConstArrayView<TensorBinding> InInputBindings, TConstArrayView<TensorBinding> InOutputBindings)
 {
-	checkf(Session.IsValid(), TEXT("FModelInstanceORT::RunSync(): Called without a Session, FModelInstanceORT::Init() should have been called."));
-
 	SCOPED_NAMED_EVENT_TEXT("FModelInstanceORTBase::RunSync", FColor::Magenta);
 
+	if (!Session.IsValid())
+	{
+		UE_LOG(LogNNE, Error, TEXT("FModelInstanceORTBase::RunSync(): Called without a Session, FModelInstanceORT::Init() should have been called."));
+		return ModelInterface::ERunSyncStatus::Fail;
+	}
+
 	// Verify the model inputs were prepared
-	if (NNE::Internal::FModelInstanceBase<ModelInterface>::InputTensorShapes.Num() == 0)
+	if (NNE::Internal::FModelInstanceBase<ModelInterface>::InputTensorShapes.IsEmpty())
+	{
+		UE_LOG(LogNNE, Error, TEXT("FModelInstanceORTBase::RunSync(): Input shapes are not set, please call SetInputTensorShapes."));
+		return ModelInterface::ERunSyncStatus::Fail;
+	}
+
+	check(NNE::Internal::FModelInstanceBase<ModelInterface>::InputTensorShapes.Num() == InputTensors.Num());
+	check(NNE::Internal::FModelInstanceBase<ModelInterface>::InputTensorShapes.Num() == InputTensorNames.Num());
+	check(NNE::Internal::FModelInstanceBase<ModelInterface>::InputSymbolicTensors.Num() == InputTensors.Num());
+
+	if (InInputBindings.Num() != InputTensors.Num())
+	{
+		UE_LOG(LogNNE, Error, TEXT("FModelInstanceORTBase::RunSync(): Input bindings need to match input tensor descriptor count (got %d, expected %d)."), InInputBindings.Num(), InputTensors.Num());
+		return ModelInterface::ERunSyncStatus::Fail;
+	}
+
+	check(NNE::Internal::FModelInstanceBase<ModelInterface>::OutputSymbolicTensors.Num() == OutputTensorNames.Num());
+
+	if (!InOutputBindings.IsEmpty() && InOutputBindings.Num() != OutputTensorNames.Num())
 	{
-		UE_LOG(LogNNE, Error, TEXT("RunSync(): Input shapes are not set, please call SetInputTensorShapes."));
+		UE_LOG(LogNNE, Error, TEXT("FModelInstanceORTBase::RunSync(): Output binding can be empty or needs to match output tensor descriptor count (got %d, expected %d)."), InOutputBindings.Num(), OutputTensorNames.Num());
 		return ModelInterface::ERunSyncStatus::Fail;
 	}
 
@@ -360,37 +397,68 @@ typename ModelInterface::ERunSyncStatus FModelInstanceORTBase<ModelInterface, Te
 	try
 #endif // WITH_EDITOR
 	{
-		TArray<Ort::Value> InputOrtTensors;
-		BindTensorsToORT(InInputBindings, InputTensors, InputTensorsORTType, *AllocatorInfo, InputOrtTensors);
-
-		if (!OutputTensors.IsEmpty())
+		TArray<Ort::Value> OrtInputTensors;
+		for (int32 i = 0; i < InputTensorNames.Num(); i++)
 		{
-			// If output shapes are known we can directly map preallocated output buffers
-			TArray<Ort::Value> OutputOrtTensors;
-			BindTensorsToORT(InOutputBindings, OutputTensors, OutputTensorsORTType, *AllocatorInfo, OutputOrtTensors);
+			const TensorBinding& Binding = InInputBindings[i];
+			const NNE::Internal::FTensor& Tensor = InputTensors[i];
 
-			Session->Run(Ort::RunOptions{ nullptr },
-				InputTensorNames.GetData(), &InputOrtTensors[0], InputTensorNames.Num(),
-				OutputTensorNames.GetData(), &OutputOrtTensors[0], OutputTensorNames.Num());
+			if (!Binding.Data && Binding.SizeInBytes != 0)
+			{
+				UE_LOG(LogNNE, Error, TEXT("FModelInstanceORTBase::RunSync(): Binding input tensor %d is not set but given size is non-zero %d."), i, Binding.SizeInBytes);
+				return ModelInterface::ERunSyncStatus::Fail;
+			}
+
+			if (Binding.SizeInBytes != Tensor.GetDataSize())
+			{
+				UE_LOG(LogNNE, Error, TEXT("FModelInstanceORTBase::RunSync(): Binding input tensor %d size does not match size given by tensor descriptor (got %d, expected %d)."), i, Binding.SizeInBytes, Tensor.GetDataSize());
+				return ModelInterface::ERunSyncStatus::Fail;
+			}
+
+			OrtInputTensors.Add(CreateTensor(*MemoryInfo, Binding, Tensor, InputTensorsORTType[i]));
 		}
-		else
+
+		TArray<Ort::Value> OrtOutputTensors;
+		for (int32 i = 0; i < OutputTensorNames.Num(); i++)
 		{
-			TArray<Ort::Value> OutputOrtTensors;
-			for (int32 i = 0; i < InOutputBindings.Num(); ++i)
+			if (OutputTensors.IsEmpty() ||
+				InOutputBindings.IsEmpty() ||
+				!InOutputBindings[i].Data ||
+				InOutputBindings[i].SizeInBytes < OutputTensors[i].GetDataSize())
 			{
-				OutputOrtTensors.Emplace(nullptr);
+				OrtOutputTensors.Emplace(nullptr);
 			}
+			else
+			{
+				OrtOutputTensors.Add(CreateTensor(*MemoryInfo, InOutputBindings[i], OutputTensors[i], OutputTensorsORTType[i]));
+			}
+		}
 
-			Session->Run(Ort::RunOptions{ nullptr },
-				InputTensorNames.GetData(), &InputOrtTensors[0], InputTensorNames.Num(),
-				OutputTensorNames.GetData(), &OutputOrtTensors[0], OutputTensorNames.Num());
+		Session->Run(Ort::RunOptions{nullptr},
+			InputTensorNames.GetData(), &OrtInputTensors[0], InputTensorNames.Num(),
+			OutputTensorNames.GetData(), &OrtOutputTensors[0], OutputTensorNames.Num());
 
-			// Output shapes were resolved during inference: Copy the data back to bindings and expose output tensor shapes
-			CopyFromORTToBindings(OutputOrtTensors, InOutputBindings, NNE::Internal::FModelInstanceBase<ModelInterface>::OutputSymbolicTensors, OutputTensors);
+		// At this (latest) stage shapes are known, therefore set them if not present yet and possibly copy data to output binding
+		if (OutputTensors.IsEmpty())
+		{
 			check(NNE::Internal::FModelInstanceBase<ModelInterface>::OutputTensorShapes.IsEmpty());
-			for (int32 i = 0; i < OutputTensors.Num(); ++i)
+
+			for (int32 i = 0; i < OutputTensorNames.Num(); i++)
 			{
-				NNE::Internal::FModelInstanceBase<ModelInterface>::OutputTensorShapes.Emplace(OutputTensors[i].GetShape());
+				const NNE::FTensorDesc& TensorDesc = NNE::Internal::FModelInstanceBase<ModelInterface>::OutputSymbolicTensors[i];
+				const NNE::FTensorShape Shape = NNE::FTensorShape::Make(OrtHelper::GetShape(OrtOutputTensors[i]));
+				const NNE::Internal::FTensor Tensor = NNE::Internal::FTensor::Make(TensorDesc.GetName(), Shape, TensorDesc.GetDataType());
+
+				OutputTensors.Add(Tensor);
+				NNE::Internal::FModelInstanceBase<ModelInterface>::OutputTensorShapes.Add(Shape);
+
+				if (!InOutputBindings.IsEmpty() &&
+					InOutputBindings[i].Data &&
+					Tensor.GetDataSize() > 0 &&
+					InOutputBindings[i].SizeInBytes >= Tensor.GetDataSize())
+				{
+					FMemory::Memcpy(InOutputBindings[i].Data, OrtOutputTensors[i].GetTensorData<void>(), Tensor.GetDataSize());
+				}
 			}
 		}
 	}
@@ -582,6 +650,9 @@ bool FModelInstanceORTDmlRDG::ConfigureTensors(const Ort::Session& ActiveSession
 	TArray<TArray<FString>>& SymbolicDimensionNames 	= bAreTensorInputs ? InputSymbolicDimensionNames : OutputSymbolicDimensionNames;
 
 	SymbolicTensorDescs.Reset();
+	TensorsORTType.Reset();
+	TensorNameValues.Reset();
+	TensorNames.Reset();
 	SymbolicDimensionNames.SetNum(NumberTensors);
 
 	for (uint32 TensorIndex = 0; TensorIndex < NumberTensors; ++TensorIndex)
@@ -759,13 +830,15 @@ FModelInstanceORTDmlRDG::ESetInputTensorShapesStatus FModelInstanceORTDmlRDG::Se
 	return ESetInputTensorShapesStatus::Ok;
 }
 
-Ort::Value CreateTensor(const OrtDmlApi* DmlApi, const Ort::MemoryInfo& MemoryInfo, FRHIBuffer* Buffer, const NNE::Internal::FTensor& Tensor, ONNXTensorElementDataType ElementDataType,
+Ort::Value CreateTensor(const OrtDmlApi& DmlApi, const Ort::MemoryInfo& MemoryInfo, FRHIBuffer* Buffer, const NNE::Internal::FTensor& Tensor, ONNXTensorElementDataType ElementDataType,
 	TArray<std::unique_ptr<void, void (*)(void*)>>& DmlAllocatorResources)
 {
+	checkf(Buffer, TEXT("CreateTensor needs Buffer to be set"));
+
 	ID3D12Resource* NativeD3D12Resource = GetID3D12DynamicRHI()->RHIGetResource(Buffer);
 
 	void* DmlAllocatorResourcePtr;
-	Ort::ThrowOnError(DmlApi->CreateGPUAllocationFromD3DResource(NativeD3D12Resource, &DmlAllocatorResourcePtr));
+	Ort::ThrowOnError(DmlApi.CreateGPUAllocationFromD3DResource(NativeD3D12Resource, &DmlAllocatorResourcePtr));
 
 	std::unique_ptr<void, void (*)(void*)> DmlAllocatorResource(DmlAllocatorResourcePtr,
 		[] (void* Ptr)
@@ -794,29 +867,71 @@ Ort::Value CreateTensor(const OrtDmlApi* DmlApi, const Ort::MemoryInfo& MemoryIn
 
 FModelInstanceORTDmlRDG::EEnqueueRDGStatus FModelInstanceORTDmlRDG::EnqueueRDG(FRDGBuilder& GraphBuilder, TConstArrayView<NNE::FTensorBindingRDG> Inputs, TConstArrayView<NNE::FTensorBindingRDG> Outputs)
 {
-	checkf(Session.IsValid(), TEXT("FModelInstanceORTDmlRDG::EnqueueRDG(): Called without a Session, FModelInstanceORTDmlRDG::Init() should have been called."));
-
 	SCOPED_NAMED_EVENT_TEXT("FModelInstanceORTDmlRDG::EnqueueRDG", FColor::Magenta);
 
-	// Verify the model inputs were prepared
-	if (InputTensorShapes.Num() == 0)
+	if (!Session.IsValid())
+	{
+		UE_LOG(LogNNE, Error, TEXT("FModelInstanceORTDmlRDG::EnqueueRDG(): Called without a Session, FModelInstanceORTDmlRDG::Init() should have been called."));
+		return EEnqueueRDGStatus::Fail;
+	}
+
+	if (InputTensorShapes.IsEmpty())
+	{
+		UE_LOG(LogNNE, Error, TEXT("FModelInstanceORTDmlRDG::EnqueueRDG(): Input shapes are not set, please call SetInputTensorShapes."));
+		return EEnqueueRDGStatus::Fail;
+	}
+
+	check(InputTensorShapes.Num() == InputTensors.Num());
+	check(InputTensorShapes.Num() == InputTensorNames.Num());
+	check(InputSymbolicTensors.Num() == InputTensors.Num());
+
+	if (Inputs.Num() != InputTensors.Num())
+	{
+		UE_LOG(LogNNE, Error, TEXT("FModelInstanceORTDmlRDG::EnqueueRDG(): Input bindings need to match input tensor descriptor count (got %d, expected %d)."), Inputs.Num(), InputTensors.Num());
+		return EEnqueueRDGStatus::Fail;
+	}
+
+	check(OutputSymbolicTensors.Num() == OutputTensorNames.Num());
+
+	if (!Outputs.IsEmpty() && Outputs.Num() != OutputTensorNames.Num())
 	{
-		UE_LOG(LogNNE, Error, TEXT("EnqueueRDG(): Input shapes are not set, please call SetInputTensorShapes."));
+		UE_LOG(LogNNE, Error, TEXT("FModelInstanceORTDmlRDG::EnqueueRDG(): Output binding can be empty or needs to match output tensor descriptor count (got %d, expected %d)."), Outputs.Num(), OutputTensorNames.Num());
 		return EEnqueueRDGStatus::Fail;
 	}
 
 	FORTModelInstanceRDGParameters* PassParameters = GraphBuilder.AllocParameters<FORTModelInstanceRDGParameters>();
-	for (const NNE::FTensorBindingRDG& Binding : Inputs)
+	for (int32 i = 0; i < Inputs.Num(); i++)
 	{
+		const NNE::FTensorBindingRDG& Binding = Inputs[i];
+		if (!Binding.Buffer && InputTensors[i].GetDataSize() != 0)
+		{
+			UE_LOG(LogNNE, Error, TEXT("FModelInstanceORTDmlRDG::EnqueueRDG(): Binding input tensor %d is not set but given size by tensor descriptor is non-zero %d."), i, InputTensors[i].GetDataSize());
+			return EEnqueueRDGStatus::Fail;
+		}
+
+		if (Binding.Buffer->Desc.GetSize() != InputTensors[i].GetDataSize())
+		{
+			UE_LOG(LogNNE, Error, TEXT("FModelInstanceORTDmlRDG::EnqueueRDG(): Binding input tensor %d size does not match size given by tensor descriptor (got %d, expected %d)."), i, Binding.Buffer->Desc.GetSize(), InputTensors[i].GetDataSize());
+			return EEnqueueRDGStatus::Fail;
+		}
+
 		PassParameters->InputBuffers.Emplace(Binding.Buffer, ERHIAccess::CopySrc);
 	}
-	for (const NNE::FTensorBindingRDG& Binding : Outputs)
+
+	TArray<int32> ValidOutputs;
+	for (int32 i = 0; i < Outputs.Num(); i++)
 	{
-		PassParameters->OutputBuffers.Emplace(Binding.Buffer, ERHIAccess::CopyDest);
+		const NNE::FTensorBindingRDG& Binding = Outputs[i];
+
+		if (Binding.Buffer && Binding.Buffer->Desc.GetSize() >= OutputTensors[i].GetDataSize())
+		{
+			PassParameters->OutputBuffers.Emplace(Binding.Buffer, ERHIAccess::CopyDest);
+			ValidOutputs.Add(i);
+		}
 	}
 
 	GraphBuilder.AddPass(RDG_EVENT_NAME("FModelInstanceORTDmlRDG::EnqueueRDG.AddPass"), PassParameters, ERDGPassFlags::Readback,
-	[this, PassParameters](FRHICommandListImmediate& RHICmdList)
+	[this, ValidOutputsCopy = ValidOutputs, PassParameters](FRHICommandListImmediate& RHICmdList)
 	{
 		SCOPED_NAMED_EVENT_TEXT("FModelInstanceORTDmlRDG::EnqueueRDG.AddPass", FColor::Magenta);
 
@@ -826,11 +941,14 @@ FModelInstanceORTDmlRDG::EEnqueueRDGStatus FModelInstanceORTDmlRDG::EnqueueRDG(F
 		{
 			InputBuffers[i] = PassParameters->InputBuffers[i]->GetRHI();
 		}
+
+		check(ValidOutputsCopy.Num() == PassParameters->OutputBuffers.Num());
+
 		TArray<FRHIBuffer*> OutputBuffers;
-		OutputBuffers.SetNumUninitialized(PassParameters->OutputBuffers.Num());
-		for (int32 i = 0; i < PassParameters->OutputBuffers.Num(); i++)
+		OutputBuffers.SetNumZeroed(OutputSymbolicTensors.Num());
+		for (int32 i = 0; i < ValidOutputsCopy.Num(); i++)
 		{
-			OutputBuffers[i] = PassParameters->OutputBuffers[i]->GetRHI();
+			OutputBuffers[ValidOutputsCopy[i]] = PassParameters->OutputBuffers[i]->GetRHI();
 		}
 
 		// Submit previous work here to the GPU to avoid ORT Session Run() dispatching its work first
@@ -853,37 +971,34 @@ FModelInstanceORTDmlRDG::EEnqueueRDGStatus FModelInstanceORTDmlRDG::EnqueueRDG(F
 					TArray<Ort::Value> OrtInputTensors;
 					TArray<Ort::Value> OrtOutputTensors;
 
-					Ort::IoBinding IoBinding = Ort::IoBinding(*Session);
-
 					for (int32 i = 0; i < InputBuffersCopyCopy.Num(); i++)
 					{
-						OrtInputTensors.Add(CreateTensor(DmlApi, MemoryInfo, InputBuffersCopyCopy[i], InputTensors[i], InputTensorsORTType[i], DmlAllocatorResources));
-
-						IoBinding.BindInput(InputTensorNames[i], OrtInputTensors[i]);
+						OrtInputTensors.Add(CreateTensor(*DmlApi, MemoryInfo, InputBuffersCopyCopy[i], InputTensors[i], InputTensorsORTType[i], DmlAllocatorResources));
 					}
 					for (int32 i = 0; i < OutputBuffersCopyCopy.Num(); i++)
 					{
-						OrtOutputTensors.Add(CreateTensor(DmlApi, MemoryInfo, OutputBuffersCopyCopy[i], OutputTensors[i], OutputTensorsORTType[i], DmlAllocatorResources));
-
-						IoBinding.BindOutput(OutputTensorNames[i], OrtOutputTensors[i]);
+						if (OutputBuffersCopyCopy[i])
+						{
+							OrtOutputTensors.Add(CreateTensor(*DmlApi, MemoryInfo, OutputBuffersCopyCopy[i], OutputTensors[i], OutputTensorsORTType[i], DmlAllocatorResources));
+						}
+						else
+						{
+							OrtOutputTensors.Emplace(nullptr);
+						}
 					}
 
-					// Don't use this sync, its CPU to GPU, but we need GPU to GPU
-					// IoBinding.SynchronizeInputs();
-
-					Session->Run({}, IoBinding);
-
-					// // Don't use this sync, its CPU to GPU, but we need GPU to GPU
-					// IoBinding.SynchronizeOutputs();
+					Session->Run(Ort::RunOptions{ nullptr },
+						InputTensorNames.GetData(), &OrtInputTensors[0], InputTensorNames.Num(),
+						OutputTensorNames.GetData(), &OrtOutputTensors[0], OutputTensorNames.Num());
 				}
 #if WITH_EDITOR
 				catch (const Ort::Exception& Exception)
 				{
-					UE_LOG(LogNNE, Error, TEXT("%s"), UTF8_TO_TCHAR(Exception.what()));
+					UE_LOG(LogNNE, Error, TEXT("FModelInstanceORTDmlRDG::EnqueueRDG(): %s"), UTF8_TO_TCHAR(Exception.what()));
 				}
 				catch (...)
 				{
-					UE_LOG(LogNNE, Error, TEXT("Unknown exception!"));
+					UE_LOG(LogNNE, Error, TEXT("FModelInstanceORTDmlRDG::EnqueueRDG(): Unknown exception!"));
 				}
 #endif // WITH_EDITOR
 			}, false);
diff --git a/Engine/Plugins/NNE/NNERuntimeORT/Source/NNERuntimeORT/Private/NNERuntimeORTModel.h b/Engine/Plugins/NNE/NNERuntimeORT/Source/NNERuntimeORT/Private/NNERuntimeORTModel.h
--- a/Engine/Plugins/NNE/NNERuntimeORT/Source/NNERuntimeORT/Private/NNERuntimeORTModel.h
+++ b/Engine/Plugins/NNE/NNERuntimeORT/Source/NNERuntimeORT/Private/NNERuntimeORTModel.h
@@ -52,7 +52,7 @@ protected:
 	TUniquePtr<Ort::Session> Session;
 	TUniquePtr<Ort::AllocatorWithDefaultOptions> Allocator;
 	TUniquePtr<Ort::SessionOptions> SessionOptions;
-	TUniquePtr<Ort::MemoryInfo> AllocatorInfo;
+	TUniquePtr<Ort::MemoryInfo> MemoryInfo;
 
 	/** IO ORT-related variables */
 	TArray<ONNXTensorElementDataType> InputTensorsORTType;
diff --git a/Engine/Plugins/NNE/NNERuntimeORT/Source/NNERuntimeORT/Private/NNERuntimeORTUtils.cpp b/Engine/Plugins/NNE/NNERuntimeORT/Source/NNERuntimeORT/Private/NNERuntimeORTUtils.cpp
--- a/Engine/Plugins/NNE/NNERuntimeORT/Source/NNERuntimeORT/Private/NNERuntimeORTUtils.cpp
+++ b/Engine/Plugins/NNE/NNERuntimeORT/Source/NNERuntimeORT/Private/NNERuntimeORTUtils.cpp
@@ -5,12 +5,38 @@
 #include "HAL/PlatformTime.h"
 #include "Misc/FileHelper.h"
 #include "Misc/Paths.h"
+#include "NNE.h"
 
 namespace UE::NNERuntimeORT::Private
 {
 namespace OrtHelper
 {
 
+TArray<uint32> GetShape(const Ort::Value& OrtTensor)
+{
+	OrtTensorTypeAndShapeInfo* TypeAndShapeInfoPtr = nullptr;
+	size_t DimensionsCount = 0;
+
+	Ort::ThrowOnError(Ort::GetApi().GetTensorTypeAndShape(OrtTensor, &TypeAndShapeInfoPtr));
+	Ort::ThrowOnError(Ort::GetApi().GetDimensionsCount(TypeAndShapeInfoPtr, &DimensionsCount));
+
+	TArray<int64_t> OrtShape;
+
+	OrtShape.SetNumUninitialized(DimensionsCount);
+	Ort::ThrowOnError(Ort::GetApi().GetDimensions(TypeAndShapeInfoPtr, OrtShape.GetData(), OrtShape.Num()));
+	Ort::GetApi().ReleaseTensorTypeAndShapeInfo(TypeAndShapeInfoPtr);
+
+	TArray<uint32> Result;
+
+	Algo::Transform(OrtShape, Result, [](int64_t Value)
+	{
+		check(Value >= 0);
+		return (uint32)Value;
+	});
+
+	return Result;
+}
+
 bool OptimizeModel(FNNEModelRaw& Model, ENNEInferenceFormat TargetFormat)
 {
 	if (Model.Format != ENNEInferenceFormat::ONNX)
@@ -77,4 +103,101 @@ bool OptimizeModel(FNNEModelRaw& Model, ENNEInferenceFormat TargetFormat)
 }
 
 } // OrtHelper
+
+TypeInfoORT TranslateTensorTypeORTToNNE(ONNXTensorElementDataType OrtDataType)
+{
+	ENNETensorDataType DataType = ENNETensorDataType::None;
+	uint64 ElementSize = 0;
+
+	switch (OrtDataType) {
+	case ONNX_TENSOR_ELEMENT_DATA_TYPE_UNDEFINED:
+		DataType = ENNETensorDataType::None;
+		ElementSize = 0;
+		break;
+
+	case ONNX_TENSOR_ELEMENT_DATA_TYPE_FLOAT:
+		DataType = ENNETensorDataType::Float;
+		ElementSize = sizeof(float);
+		break;
+
+	case ONNX_TENSOR_ELEMENT_DATA_TYPE_UINT8:
+		DataType = ENNETensorDataType::UInt8;
+		ElementSize = sizeof(uint8);
+		break;
+	case ONNX_TENSOR_ELEMENT_DATA_TYPE_INT8:
+		DataType = ENNETensorDataType::Int8;
+		ElementSize = sizeof(int8);
+		break;
+	case ONNX_TENSOR_ELEMENT_DATA_TYPE_UINT16:
+		DataType = ENNETensorDataType::UInt16;
+		ElementSize = sizeof(uint16);
+		break;
+	case ONNX_TENSOR_ELEMENT_DATA_TYPE_INT16:
+		DataType = ENNETensorDataType::Int16;
+		ElementSize = sizeof(int16);
+		break;
+
+	case ONNX_TENSOR_ELEMENT_DATA_TYPE_INT32:
+		DataType = ENNETensorDataType::Int32;
+		ElementSize = sizeof(int32);
+		break;
+
+	case ONNX_TENSOR_ELEMENT_DATA_TYPE_INT64:
+		DataType = ENNETensorDataType::Int64;
+		ElementSize = sizeof(int64);
+		break;
+
+	case ONNX_TENSOR_ELEMENT_DATA_TYPE_STRING:
+		DataType = ENNETensorDataType::Char;
+		ElementSize = sizeof(char);
+		break;
+
+	case ONNX_TENSOR_ELEMENT_DATA_TYPE_BOOL:
+		DataType = ENNETensorDataType::Boolean;
+		ElementSize = sizeof(bool);
+		break;
+
+	case ONNX_TENSOR_ELEMENT_DATA_TYPE_FLOAT16:
+		DataType = ENNETensorDataType::Half;
+		ElementSize = 2;
+		break;
+
+	case ONNX_TENSOR_ELEMENT_DATA_TYPE_DOUBLE:
+		DataType = ENNETensorDataType::Double;
+		ElementSize = sizeof(double);
+		break;
+
+	case ONNX_TENSOR_ELEMENT_DATA_TYPE_UINT32:
+		DataType = ENNETensorDataType::UInt32;
+		ElementSize = sizeof(uint32);
+		break;
+
+	case ONNX_TENSOR_ELEMENT_DATA_TYPE_UINT64:
+		DataType = ENNETensorDataType::UInt64;
+		ElementSize = sizeof(uint64);
+		break;
+
+	case ONNX_TENSOR_ELEMENT_DATA_TYPE_COMPLEX64:
+		DataType = ENNETensorDataType::Complex64;
+		ElementSize = 8;
+		break;
+
+	case ONNX_TENSOR_ELEMENT_DATA_TYPE_COMPLEX128:
+		DataType = ENNETensorDataType::Complex128;
+		ElementSize = 16;
+		break;
+
+	case ONNX_TENSOR_ELEMENT_DATA_TYPE_BFLOAT16:
+		DataType = ENNETensorDataType::BFloat16;
+		ElementSize = 2;
+		break;
+
+	default:
+		DataType = ENNETensorDataType::None;
+		break;
+	}
+
+	return TypeInfoORT{ DataType, ElementSize };
+}
+
 } // namespace UE::NNERuntimeORT::Private
\ No newline at end of file
diff --git a/Engine/Plugins/NNE/NNERuntimeORT/Source/NNERuntimeORT/Private/NNERuntimeORTUtils.h b/Engine/Plugins/NNE/NNERuntimeORT/Source/NNERuntimeORT/Private/NNERuntimeORTUtils.h
--- a/Engine/Plugins/NNE/NNERuntimeORT/Source/NNERuntimeORT/Private/NNERuntimeORTUtils.h
+++ b/Engine/Plugins/NNE/NNERuntimeORT/Source/NNERuntimeORT/Private/NNERuntimeORTUtils.h
@@ -2,43 +2,16 @@
 
 #pragma once
 
-#include "HAL/UnrealMemory.h"
-#include "NNE.h"
+#include "Containers/Array.h"
 #include "NNEOnnxruntime.h"
-#include "NNERuntimeCPU.h"
 #include "NNERuntimeFormat.h"
-#include "NNERuntimeGPU.h"
-#include "NNETensor.h"
 #include "NNETypes.h"
 
 namespace UE::NNERuntimeORT::Private
 {
 	namespace OrtHelper
 	{
-		inline TArray<uint32> GetShape(const Ort::Value& OrtTensor)
-		{
-			OrtTensorTypeAndShapeInfo* TypeAndShapeInfoPtr = nullptr;
-			size_t DimensionsCount = 0;
-
-			Ort::ThrowOnError(Ort::GetApi().GetTensorTypeAndShape(OrtTensor, &TypeAndShapeInfoPtr));
-			Ort::ThrowOnError(Ort::GetApi().GetDimensionsCount(TypeAndShapeInfoPtr, &DimensionsCount));
-
-			TArray<int64_t> OrtShape;
-
-			OrtShape.SetNumUninitialized(DimensionsCount);
-			Ort::ThrowOnError(Ort::GetApi().GetDimensions(TypeAndShapeInfoPtr, OrtShape.GetData(), OrtShape.Num()));
-			Ort::GetApi().ReleaseTensorTypeAndShapeInfo(TypeAndShapeInfoPtr);
-
-			TArray<uint32> Result;
-
-			Algo::Transform(OrtShape, Result, [](int64_t Value)
-			{
-				check(Value >= 0);
-				return (uint32)Value;
-			});
-
-			return Result;
-		}
+		TArray<uint32> GetShape(const Ort::Value& OrtTensor);
 
 		bool OptimizeModel(FNNEModelRaw& Model, ENNEInferenceFormat OutFormat);
 	}
@@ -49,185 +22,6 @@ namespace UE::NNERuntimeORT::Private
 		uint64 ElementSize = 0;
 	};
 
-	inline TypeInfoORT TranslateTensorTypeORTToNNE(ONNXTensorElementDataType OrtDataType)
-	{
-		ENNETensorDataType DataType = ENNETensorDataType::None;
-		uint64 ElementSize = 0;
-
-		switch (OrtDataType) {
-		case ONNX_TENSOR_ELEMENT_DATA_TYPE_UNDEFINED:
-			DataType = ENNETensorDataType::None;
-			ElementSize = 0;
-			break;
-
-		case ONNX_TENSOR_ELEMENT_DATA_TYPE_FLOAT:
-			DataType = ENNETensorDataType::Float;
-			ElementSize = sizeof(float);
-			break;
-
-		case ONNX_TENSOR_ELEMENT_DATA_TYPE_UINT8:
-			DataType = ENNETensorDataType::UInt8;
-			ElementSize = sizeof(uint8);
-			break;
-		case ONNX_TENSOR_ELEMENT_DATA_TYPE_INT8:
-			DataType = ENNETensorDataType::Int8;
-			ElementSize = sizeof(int8);
-			break;
-		case ONNX_TENSOR_ELEMENT_DATA_TYPE_UINT16:
-			DataType = ENNETensorDataType::UInt16;
-			ElementSize = sizeof(uint16);
-			break;
-		case ONNX_TENSOR_ELEMENT_DATA_TYPE_INT16:
-			DataType = ENNETensorDataType::Int16;
-			ElementSize = sizeof(int16);
-			break;
-
-		case ONNX_TENSOR_ELEMENT_DATA_TYPE_INT32:
-			DataType = ENNETensorDataType::Int32;
-			ElementSize = sizeof(int32);
-			break;
-
-		case ONNX_TENSOR_ELEMENT_DATA_TYPE_INT64:
-			DataType = ENNETensorDataType::Int64;
-			ElementSize = sizeof(int64);
-			break;
-
-		case ONNX_TENSOR_ELEMENT_DATA_TYPE_STRING:
-			DataType = ENNETensorDataType::Char;
-			ElementSize = sizeof(char);
-			break;
-
-		case ONNX_TENSOR_ELEMENT_DATA_TYPE_BOOL:
-			DataType = ENNETensorDataType::Boolean;
-			ElementSize = sizeof(bool);
-			break;
-
-		case ONNX_TENSOR_ELEMENT_DATA_TYPE_FLOAT16:
-			DataType = ENNETensorDataType::Half;
-			ElementSize = 2;
-			break;
-
-		case ONNX_TENSOR_ELEMENT_DATA_TYPE_DOUBLE:
-			DataType = ENNETensorDataType::Double;
-			ElementSize = sizeof(double);
-			break;
-
-		case ONNX_TENSOR_ELEMENT_DATA_TYPE_UINT32:
-			DataType = ENNETensorDataType::UInt32;
-			ElementSize = sizeof(uint32);
-			break;
-
-		case ONNX_TENSOR_ELEMENT_DATA_TYPE_UINT64:
-			DataType = ENNETensorDataType::UInt64;
-			ElementSize = sizeof(uint64);
-			break;
-
-		case ONNX_TENSOR_ELEMENT_DATA_TYPE_COMPLEX64:
-			DataType = ENNETensorDataType::Complex64;
-			ElementSize = 8;
-			break;
-
-		case ONNX_TENSOR_ELEMENT_DATA_TYPE_COMPLEX128:
-			DataType = ENNETensorDataType::Complex128;
-			ElementSize = 16;
-			break;
-
-		case ONNX_TENSOR_ELEMENT_DATA_TYPE_BFLOAT16:
-			DataType = ENNETensorDataType::BFloat16;
-			ElementSize = 2;
-			break;
-
-		default:
-			DataType = ENNETensorDataType::None;
-			break;
-		}
-
-		return TypeInfoORT{ DataType, ElementSize };
-	}
-
-	template <class TensorBinding> void BindTensorsToORT(
-		TConstArrayView<TensorBinding> InBindingTensors,
-		TConstArrayView<NNE::Internal::FTensor> InTensors,
-		TConstArrayView<ONNXTensorElementDataType> InTensorsORTType,
-		const Ort::MemoryInfo& InAllocatorInfo,
-		TArray<Ort::Value>& OutOrtTensors
-	)
-	{
-		const uint32 NumBinding = InBindingTensors.Num();
-		const uint32 NumTensors = InTensors.Num();
-
-		if (NumBinding != NumTensors)
-		{
-			UE_LOG(LogNNE, Error, TEXT("BindTensorsToORT: Number of Bindings is different from Tensors."));
-			return;
-		}
-
-		for (uint32 Index = 0; Index < NumBinding; ++Index)
-		{
-			const TensorBinding& Binding = InBindingTensors[Index];
-			const NNE::Internal::FTensor& Tensor = InTensors[Index];
-			const ONNXTensorElementDataType CurrentORTType = InTensorsORTType[Index];
-			const uint64 ByteCount{ InTensors[Index].GetDataSize() };
-			const uint32 ArrayDimensions{ (uint32)Tensor.GetShape().Rank() };
-
-			TUniquePtr<int64_t[]> SizesInt64t = MakeUnique<int64_t[]>(Tensor.GetShape().Rank());
-			for (int32 DimIndex = 0; DimIndex < Tensor.GetShape().Rank(); ++DimIndex)
-			{
-				SizesInt64t.Get()[DimIndex] = Tensor.GetShape().GetData()[DimIndex];
-			}
-			
-			OutOrtTensors.Emplace(
-				Ort::Value::CreateTensor(
-					InAllocatorInfo,
-					Binding.Data,
-					ByteCount,
-					SizesInt64t.Get(),
-					ArrayDimensions,
-					CurrentORTType
-				)
-			);
-		}
-	}
-
-	template <class TensorBinding> void CopyFromORTToBindings(
-		TConstArrayView<Ort::Value> InOrtTensors,
-		TConstArrayView<TensorBinding> InBindingTensors,
-		TConstArrayView<NNE::FTensorDesc> InTensorDescs,
-		TArray<NNE::Internal::FTensor>& OutTensors
-	)
-	{
-		const uint32 NumBinding = InOrtTensors.Num();
-		const uint32 NumDescriptors = InTensorDescs.Num();
-
-		if (NumBinding != NumDescriptors)
-		{
-			UE_LOG(LogNNE, Error, TEXT("CopyFromORTToBindings: Number of Bindings is different from Descriptors."));
-			return;
-		}
-
-		for (uint32 Index = 0; Index < NumBinding; ++Index)
-		{
-			const TensorBinding& Binding = InBindingTensors[Index];
-			const NNE::FTensorDesc& TensorDesc = InTensorDescs[Index];
-			const Ort::Value& OrtTensor = InOrtTensors[Index];
-			const TArray<uint32> OrtShape = OrtHelper::GetShape(OrtTensor);
-			const void* OrtTensorData = OrtTensor.GetTensorData<void>();
-			const NNE::FTensorShape Shape = NNE::FTensorShape::Make(OrtShape);
-			const NNE::Internal::FTensor Tensor = NNE::Internal::FTensor::Make(TensorDesc.GetName(), Shape, TensorDesc.GetDataType());
-			const uint64 DataSize = FMath::Min(Tensor.GetDataSize(), Binding.SizeInBytes);
-
-			OutTensors.Emplace(Tensor);
-
-			if (DataSize > 0)
-			{
-				FMemory::Memcpy(Binding.Data, OrtTensorData, DataSize);
-			}
-
-			if (Tensor.GetDataSize() > Binding.SizeInBytes)
-			{
-				UE_LOG(LogNNE, Error, TEXT("CopyFromORTToBindings: Binding buffer was not large enough to contain all of the data, only first %d bytes were copied."), DataSize);
-			}
-		}
-	}
+	TypeInfoORT TranslateTensorTypeORTToNNE(ONNXTensorElementDataType OrtDataType);
 
 } // UE::NNERuntimeORT::Private
\ No newline at end of file
