[NNE] RuntimeORT/HLSL change to ONNX Runtime (ORT) global thread pool in Editor to improve import/cook performance (per OrtSession thread pool creation/destruction is very expensive).

#tests Editor Win64 all tests and package standalone

diff --git a/Engine/Plugins/Experimental/NNERuntimeRDG/Source/NNERuntimeRDGUtils/Private/NNERuntimeRDGUtilsModelOptimizerONNX.cpp b/Engine/Plugins/Experimental/NNERuntimeRDG/Source/NNERuntimeRDGUtils/Private/NNERuntimeRDGUtilsModelOptimizerONNX.cpp
--- a/Engine/Plugins/Experimental/NNERuntimeRDG/Source/NNERuntimeRDGUtils/Private/NNERuntimeRDGUtilsModelOptimizerONNX.cpp
+++ b/Engine/Plugins/Experimental/NNERuntimeRDG/Source/NNERuntimeRDGUtils/Private/NNERuntimeRDGUtilsModelOptimizerONNX.cpp
@@ -49,20 +49,25 @@ public:
 		//We only enable all the optimization when going to ORT format itself for the CPU provider
 		GraphOptimizationLevel OptimizationLevel = TargetFormat == ENNEInferenceFormat::ONNX ? ORT_ENABLE_BASIC : ORT_ENABLE_ALL;
 
-		double ONNXModelOptimisationStartTime = FPlatformTime::Seconds();
-		
 		FFileHelper::SaveArrayToFile(Model.Data, *ModelToOptimizePath);
 		{
-			Ort::Env Env(ORT_LOGGING_LEVEL_INFO);
-			Ort::SessionOptions SessOptions;
+			Ort::ThreadingOptions ThreadingOptions;
+			ThreadingOptions.SetGlobalIntraOpNumThreads(1);
+			ThreadingOptions.SetGlobalInterOpNumThreads(1);
+
+			Ort::Env Env(ThreadingOptions);
 
-			SessOptions.SetGraphOptimizationLevel(OptimizationLevel);
+			Ort::SessionOptions SessionOptions;
+			SessionOptions.DisablePerSessionThreads();
+			SessionOptions.SetGraphOptimizationLevel(OptimizationLevel);
 			#if PLATFORM_WINDOWS
-				SessOptions.SetOptimizedModelFilePath(*ModelOptimizedPath);
-				Ort::Session Session(Env, *ModelToOptimizePath, SessOptions);
+				SessionOptions.SetOptimizedModelFilePath(*ModelOptimizedPath);
+
+				Ort::Session Session(Env, *ModelToOptimizePath, SessionOptions);
 			#else
-				SessOptions.SetOptimizedModelFilePath(TCHAR_TO_ANSI(*ModelOptimizedPath));
-				Ort::Session Session(Env, TCHAR_TO_ANSI(*ModelToOptimizePath), SessOptions);
+				SessionOptions.SetOptimizedModelFilePath(TCHAR_TO_ANSI(*ModelOptimizedPath));
+
+				Ort::Session Session(Env, TCHAR_TO_ANSI(*ModelToOptimizePath), SessionOptions);
 			#endif
 		}
 		FFileHelper::LoadFileToArray(Model.Data, *ModelOptimizedPath);
@@ -70,11 +75,6 @@ public:
 		IFileManager::Get().Delete(*ModelToOptimizePath);
 		IFileManager::Get().Delete(*ModelOptimizedPath);
 
-		double ONNXModelOptimisationEndTime = FPlatformTime::Seconds();
-		float ONNXModelOptimisationTime = static_cast<float>(ONNXModelOptimisationEndTime - ONNXModelOptimisationStartTime);
-
-		UE_LOG(LogNNE, Display, TEXT("OnnxRuntimeModelOptimizerPass runned in %0.1f seconds."), ONNXModelOptimisationTime);
-
 		Model.Format = TargetFormat;
 
 		return true;
diff --git a/Engine/Plugins/NNE/NNERuntimeORT/Source/NNERuntimeORT/NNERuntimeORT.Build.cs b/Engine/Plugins/NNE/NNERuntimeORT/Source/NNERuntimeORT/NNERuntimeORT.Build.cs
--- a/Engine/Plugins/NNE/NNERuntimeORT/Source/NNERuntimeORT/NNERuntimeORT.Build.cs
+++ b/Engine/Plugins/NNE/NNERuntimeORT/Source/NNERuntimeORT/NNERuntimeORT.Build.cs
@@ -17,15 +17,15 @@ public class NNERuntimeORT : ModuleRules
 			"NNEOnnxruntime",
 			"Projects",
 			"RenderCore",
-			"DeveloperSettings"
+			"DeveloperSettings",
+			"RHI"
 		});
 
 		if (Target.Platform == UnrealTargetPlatform.Win64)
 		{
 			PrivateDependencyModuleNames.AddRange(new string[]
 			{
-				"D3D12RHI",
-				"RHI"
+				"D3D12RHI"
 			});
 
 			AddEngineThirdPartyPrivateStaticDependencies(Target, new string[]
diff --git a/Engine/Plugins/NNE/NNERuntimeORT/Source/NNERuntimeORT/Private/NNERuntimeORT.cpp b/Engine/Plugins/NNE/NNERuntimeORT/Source/NNERuntimeORT/Private/NNERuntimeORT.cpp
--- a/Engine/Plugins/NNE/NNERuntimeORT/Source/NNERuntimeORT/Private/NNERuntimeORT.cpp
+++ b/Engine/Plugins/NNE/NNERuntimeORT/Source/NNERuntimeORT/Private/NNERuntimeORT.cpp
@@ -38,7 +38,7 @@ TSharedPtr<UE::NNE::FSharedModelData> UNNERuntimeORTCpu::CreateModelData(const F
 	}
 
 	FNNEModelRaw InputModel{TArray<uint8>{FileData}, ENNEInferenceFormat::ONNX};
-	if (!UE::NNERuntimeORT::Private::OrtHelper::OptimizeModel(InputModel, ENNEInferenceFormat::ORT))
+	if (!UE::NNERuntimeORT::Private::OrtHelper::OptimizeModel(Environment.ToSharedRef(), InputModel, ENNEInferenceFormat::ORT))
 	{
 		return {};
 	}
@@ -149,7 +149,7 @@ TSharedPtr<UE::NNE::FSharedModelData> UNNERuntimeORTDml::CreateModelData(const F
 	}
 
 	FNNEModelRaw InputModel{TArray<uint8>{FileData}, ENNEInferenceFormat::ONNX};
-	if (!UE::NNERuntimeORT::Private::OrtHelper::OptimizeModel(InputModel, ENNEInferenceFormat::ONNX))
+	if (!UE::NNERuntimeORT::Private::OrtHelper::OptimizeModel(Environment.ToSharedRef(), InputModel, ENNEInferenceFormat::ONNX))
 	{
 		return {};
 	}
diff --git a/Engine/Plugins/NNE/NNERuntimeORT/Source/NNERuntimeORT/Private/NNERuntimeORTEnv.cpp b/Engine/Plugins/NNE/NNERuntimeORT/Source/NNERuntimeORT/Private/NNERuntimeORTEnv.cpp
--- a/Engine/Plugins/NNE/NNERuntimeORT/Source/NNERuntimeORT/Private/NNERuntimeORTEnv.cpp
+++ b/Engine/Plugins/NNE/NNERuntimeORT/Source/NNERuntimeORT/Private/NNERuntimeORTEnv.cpp
@@ -13,11 +13,11 @@ void FEnvironment::CreateOrtEnv(const FConfig& InConfig)
 #endif // WITH_EDITOR
 	checkf(OrtEnvironment == Ort::Env(nullptr), TEXT("Ort::Env already created!"));
 
-	if (Config.bUseGlobalThreadPool)
+	if (InConfig.bUseGlobalThreadPool)
 	{
 		Ort::ThreadingOptions ThreadingOptions;
-		ThreadingOptions.SetGlobalIntraOpNumThreads(Config.IntraOpNumThreads);
-		ThreadingOptions.SetGlobalInterOpNumThreads(Config.InterOpNumThreads);
+		ThreadingOptions.SetGlobalIntraOpNumThreads(InConfig.IntraOpNumThreads);
+		ThreadingOptions.SetGlobalInterOpNumThreads(InConfig.InterOpNumThreads);
 
 		// Calls OrtApi::CreateEnvWithGlobalThreadPools(), needs to be called in conjunction
 		// with OrtApi::DisablePerSessionThreads or else the session will use its own thread pools.
diff --git a/Engine/Plugins/NNE/NNERuntimeORT/Source/NNERuntimeORT/Private/NNERuntimeORTEnv.h b/Engine/Plugins/NNE/NNERuntimeORT/Source/NNERuntimeORT/Private/NNERuntimeORTEnv.h
--- a/Engine/Plugins/NNE/NNERuntimeORT/Source/NNERuntimeORT/Private/NNERuntimeORTEnv.h
+++ b/Engine/Plugins/NNE/NNERuntimeORT/Source/NNERuntimeORT/Private/NNERuntimeORTEnv.h
@@ -25,6 +25,8 @@ public:
 
 	const Ort::Env& GetOrtEnv() const;
 
+	FConfig GetConfig() const { return Config; }
+
 private:
 	FConfig Config{};
 	Ort::Env OrtEnvironment{nullptr};
diff --git a/Engine/Plugins/NNE/NNERuntimeORT/Source/NNERuntimeORT/Private/NNERuntimeORTModel.cpp b/Engine/Plugins/NNE/NNERuntimeORT/Source/NNERuntimeORT/Private/NNERuntimeORTModel.cpp
--- a/Engine/Plugins/NNE/NNERuntimeORT/Source/NNERuntimeORT/Private/NNERuntimeORTModel.cpp
+++ b/Engine/Plugins/NNE/NNERuntimeORT/Source/NNERuntimeORT/Private/NNERuntimeORTModel.cpp
@@ -35,6 +35,8 @@ static TAutoConsoleVariable<bool> CVarNNERuntimeORTEnableProfiling(
 	TEXT("More information can be found at https://onnxruntime.ai/docs/performance/tune-performance/profiling-tools.html\n"),
 	ECVF_Default);
 
+DECLARE_GPU_STAT_NAMED(FNNERuntimeORTDmlRDG, TEXT("FModelInstanceORTDmlRDG::EnqueueRDG"));
+
 namespace UE::NNERuntimeORT::Private
 {
 
@@ -520,8 +522,9 @@ bool FModelInstanceORTCpu::InitializedAndConfigureMembers()
 #if PLATFORM_WINDOWS
 TSharedPtr<NNE::IModelInstanceGPU> FModelORTDmlGPU::CreateModelInstanceGPU()
 {
-	const FRuntimeConf Conf = Detail::MakeRuntimeConfigFromSettings(GetDefault<UNNERuntimeORTSettings>());
-	FModelInstanceORTDmlGPU* ModelInstance = new FModelInstanceORTDmlGPU(Conf, Environment);
+	const FRuntimeConf RuntimeConfig = Detail::MakeRuntimeConfigFromSettings(GetDefault<UNNERuntimeORTSettings>());
+
+	FModelInstanceORTDmlGPU* ModelInstance = new FModelInstanceORTDmlGPU(RuntimeConfig, Environment);
 
 	check(ModelData.IsValid());
 	if (!ModelInstance->Init(ModelData->GetView()))
@@ -562,7 +565,9 @@ FModelORTDmlRDG::FModelORTDmlRDG(TSharedRef<FEnvironment> InEnvironment, TShared
 
 TSharedPtr<NNE::IModelInstanceRDG> FModelORTDmlRDG::CreateModelInstanceRDG()
 {
-	TSharedPtr<FModelInstanceORTDmlRDG> ModelInstance = MakeShared<FModelInstanceORTDmlRDG>(ModelData, FRuntimeConf{}, Environment);
+	const FRuntimeConf RuntimeConfig = Detail::MakeRuntimeConfigFromSettings(GetDefault<UNNERuntimeORTSettings>());
+
+	TSharedPtr<FModelInstanceORTDmlRDG> ModelInstance = MakeShared<FModelInstanceORTDmlRDG>(ModelData, RuntimeConfig, Environment);
 	if (!ModelInstance->Init())
 	{
 		return {};
@@ -642,6 +647,8 @@ bool FModelInstanceORTDmlRDG::ConfigureTensors(const Ort::Session& ActiveSession
 
 bool FModelInstanceORTDmlRDG::ConfigureTensors(const Ort::Session& ActiveSession, bool bAreTensorInputs)
 {
+	SCOPED_NAMED_EVENT_TEXT("FModelInstanceORTDmlRDG::ConfigureTensors", FColor::Magenta);
+
 	const uint32 NumberTensors							= bAreTensorInputs ? ActiveSession.GetInputCount() : ActiveSession.GetOutputCount();
 	TArray<NNE::FTensorDesc>& SymbolicTensorDescs		= bAreTensorInputs ? InputSymbolicTensors   : OutputSymbolicTensors;
 	TArray<ONNXTensorElementDataType>& TensorsORTType	= bAreTensorInputs ? InputTensorsORTType	: OutputTensorsORTType;
@@ -702,6 +709,8 @@ bool FModelInstanceORTDmlRDG::ConfigureTensors(const Ort::Session& ActiveSession
 
 FModelInstanceORTDmlRDG::ESetInputTensorShapesStatus FModelInstanceORTDmlRDG::SetInputTensorShapes(TConstArrayView<NNE::FTensorShape> InInputShapes)
 {
+	SCOPED_NAMED_EVENT_TEXT("FModelInstanceORTDmlRDG::SetInputTensorShapes", FColor::Magenta);
+
 	InputTensors.Reset();
 	OutputTensors.Reset();
 	OutputTensorShapes.Reset();
@@ -930,6 +939,9 @@ FModelInstanceORTDmlRDG::EEnqueueRDGStatus FModelInstanceORTDmlRDG::EnqueueRDG(F
 		}
 	}
 
+	RDG_EVENT_SCOPE(GraphBuilder, "FModelInstanceORTDmlRDG::EnqueueRDG");
+	RDG_GPU_STAT_SCOPE(GraphBuilder, FNNERuntimeORTDmlRDG);
+
 	GraphBuilder.AddPass(RDG_EVENT_NAME("FModelInstanceORTDmlRDG::EnqueueRDG.AddPass"), PassParameters, ERDGPassFlags::Readback,
 	[this, ValidOutputsCopy = ValidOutputs, PassParameters](FRHICommandListImmediate& RHICmdList)
 	{
diff --git a/Engine/Plugins/NNE/NNERuntimeORT/Source/NNERuntimeORT/Private/NNERuntimeORTSettings.h b/Engine/Plugins/NNE/NNERuntimeORT/Source/NNERuntimeORT/Private/NNERuntimeORTSettings.h
--- a/Engine/Plugins/NNE/NNERuntimeORT/Source/NNERuntimeORT/Private/NNERuntimeORTSettings.h
+++ b/Engine/Plugins/NNE/NNERuntimeORT/Source/NNERuntimeORT/Private/NNERuntimeORTSettings.h
@@ -55,7 +55,7 @@ public:
 
 	/** Threading options in Editor targets */
 	UPROPERTY(Config, EditAnywhere, Category = "ONNX Runtime", meta = (DisplayName = "Editor threading options", ToolTip = "Threading options in Editor targets", ConfigRestartRequired = true))
-	FThreadingOptions EditorThreadingOptions{false, 0, 0, EExecutionMode::SEQUENTIAL};
+	FThreadingOptions EditorThreadingOptions{true, 0, 0, EExecutionMode::SEQUENTIAL};
 
 	/** Threading options in Non-Editor (Game, Program, ...) targets */
 	UPROPERTY(Config, EditAnywhere, Category = "ONNX Runtime", meta = (DisplayName = "Game threading options", ToolTip = "Threading options in Non-Editor (Game, Program, ...) targets", ConfigRestartRequired = true))
diff --git a/Engine/Plugins/NNE/NNERuntimeORT/Source/NNERuntimeORT/Private/NNERuntimeORTUtils.cpp b/Engine/Plugins/NNE/NNERuntimeORT/Source/NNERuntimeORT/Private/NNERuntimeORTUtils.cpp
--- a/Engine/Plugins/NNE/NNERuntimeORT/Source/NNERuntimeORT/Private/NNERuntimeORTUtils.cpp
+++ b/Engine/Plugins/NNE/NNERuntimeORT/Source/NNERuntimeORT/Private/NNERuntimeORTUtils.cpp
@@ -6,6 +6,7 @@
 #include "Misc/FileHelper.h"
 #include "Misc/Paths.h"
 #include "NNE.h"
+#include "NNERuntimeORTEnv.h"
 
 namespace UE::NNERuntimeORT::Private
 {
@@ -37,8 +38,10 @@ TArray<uint32> GetShape(const Ort::Value& OrtTensor)
 	return Result;
 }
 
-bool OptimizeModel(FNNEModelRaw& Model, ENNEInferenceFormat TargetFormat)
+bool OptimizeModel(TSharedRef<FEnvironment> InEnvironment, FNNEModelRaw& Model, ENNEInferenceFormat TargetFormat)
 {
+	SCOPED_NAMED_EVENT_TEXT("OrtHelper::OptimizeModel", FColor::Magenta);
+
 	if (Model.Format != ENNEInferenceFormat::ONNX)
 	{
 		UE_LOG(LogNNE, Warning, TEXT("NNERuntimeORT: ONNX Runtime Model Optimizer is expecting a model in ONNX format but received %u."), Model.Format);
@@ -54,24 +57,31 @@ bool OptimizeModel(FNNEModelRaw& Model, ENNEInferenceFormat TargetFormat)
 	//We only enable all the optimization when going to ORT format itself for the CPU provider
 	GraphOptimizationLevel OptimizationLevel = TargetFormat == ENNEInferenceFormat::ONNX ? ORT_ENABLE_BASIC : ORT_ENABLE_ALL;
 
-	double ONNXModelOptimisationStartTime = FPlatformTime::Seconds();
-	
 	FFileHelper::SaveArrayToFile(Model.Data, *ModelToOptimizePath);
 
 #if WITH_EDITOR
 	try
 #endif // WITH_EDITOR
 	{
-		Ort::Env Env(ORT_LOGGING_LEVEL_INFO);
-		Ort::SessionOptions SessOptions;
-
-		SessOptions.SetGraphOptimizationLevel(OptimizationLevel);
+		Ort::SessionOptions SessionOptions;
+		if (ensureMsgf(InEnvironment->GetConfig().bUseGlobalThreadPool, TEXT("For Model Optimizer use ONNX Runtime global threadpool to improve performance!")))
+		{
+			SessionOptions.DisablePerSessionThreads();
+		}
+		else
+		{
+			SessionOptions.SetIntraOpNumThreads(InEnvironment->GetConfig().IntraOpNumThreads);
+			SessionOptions.SetInterOpNumThreads(InEnvironment->GetConfig().InterOpNumThreads);
+		}
+		SessionOptions.SetGraphOptimizationLevel(OptimizationLevel);
 #if PLATFORM_WINDOWS
-		SessOptions.SetOptimizedModelFilePath(*ModelOptimizedPath);
-		Ort::Session Session(Env, *ModelToOptimizePath, SessOptions);
+		SessionOptions.SetOptimizedModelFilePath(*ModelOptimizedPath);
+
+		Ort::Session Session(InEnvironment->GetOrtEnv(), *ModelToOptimizePath, SessionOptions);
 #else
-		SessOptions.SetOptimizedModelFilePath(TCHAR_TO_ANSI(*ModelOptimizedPath));
-		Ort::Session Session(Env, TCHAR_TO_ANSI(*ModelToOptimizePath), SessOptions);
+		SessionOptions.SetOptimizedModelFilePath(TCHAR_TO_ANSI(*ModelOptimizedPath));
+		
+		Ort::Session Session(InEnvironment->GetOrtEnv(), TCHAR_TO_ANSI(*ModelToOptimizePath), SessionOptions);
 #endif
 	}
 #if WITH_EDITOR
@@ -92,11 +102,6 @@ bool OptimizeModel(FNNEModelRaw& Model, ENNEInferenceFormat TargetFormat)
 	IFileManager::Get().Delete(*ModelToOptimizePath);
 	IFileManager::Get().Delete(*ModelOptimizedPath);
 
-	double ONNXModelOptimisationEndTime = FPlatformTime::Seconds();
-	float ONNXModelOptimisationTime = static_cast<float>(ONNXModelOptimisationEndTime - ONNXModelOptimisationStartTime);
-
-	UE_LOG(LogNNE, Display, TEXT("NNERuntimeORT: ONNX Runtime Model Optimizer runned in %0.1f seconds."), ONNXModelOptimisationTime);
-
 	Model.Format = TargetFormat;
 
 	return true;
diff --git a/Engine/Plugins/NNE/NNERuntimeORT/Source/NNERuntimeORT/Private/NNERuntimeORTUtils.h b/Engine/Plugins/NNE/NNERuntimeORT/Source/NNERuntimeORT/Private/NNERuntimeORTUtils.h
--- a/Engine/Plugins/NNE/NNERuntimeORT/Source/NNERuntimeORT/Private/NNERuntimeORTUtils.h
+++ b/Engine/Plugins/NNE/NNERuntimeORT/Source/NNERuntimeORT/Private/NNERuntimeORTUtils.h
@@ -9,11 +9,14 @@
 
 namespace UE::NNERuntimeORT::Private
 {
+
+	class FEnvironment;
+
 	namespace OrtHelper
 	{
 		TArray<uint32> GetShape(const Ort::Value& OrtTensor);
 
-		bool OptimizeModel(FNNEModelRaw& Model, ENNEInferenceFormat OutFormat);
+		bool OptimizeModel(TSharedRef<FEnvironment> InEnvironment, FNNEModelRaw& Model, ENNEInferenceFormat OutFormat);
 	}
 
 	struct TypeInfoORT
