[NNE] NNERuntimeORT Graph Optimization rework.

Problem:
- ORT prints warning on cooking because of hardware specific optimization
- Reevaluate optimization flags with detailed time (specifically ORT session creation) measureemnts
- No DML specific optimizations

Solution:
- ORT CPU: reduce optimization level for cooking to non-hardware specific and revert to use ONNX format (faster)
- ORT DML: leave optimization level unchanged
- Disable optimization in Editor (consider adding cache and re-enable)

#tests Editor & Game, Win64 all tests

#rnx

diff --git a/Engine/Plugins/NNE/NNERuntimeORT/Source/NNERuntimeORT/Private/NNERuntimeORT.cpp b/Engine/Plugins/NNE/NNERuntimeORT/Source/NNERuntimeORT/Private/NNERuntimeORT.cpp
--- a/Engine/Plugins/NNE/NNERuntimeORT/Source/NNERuntimeORT/Private/NNERuntimeORT.cpp
+++ b/Engine/Plugins/NNE/NNERuntimeORT/Source/NNERuntimeORT/Private/NNERuntimeORT.cpp
@@ -31,6 +31,8 @@ UNNERuntimeORTCpu::ECanCreateModelDataStatus UNNERuntimeORTCpu::CanCreateModelDa
 
 TSharedPtr<UE::NNE::FSharedModelData> UNNERuntimeORTCpu::CreateModelData(const FString& FileType, TConstArrayView<uint8> FileData, const TMap<FString, TConstArrayView<uint8>>& AdditionalFileData, const FGuid& FileId, const ITargetPlatform* TargetPlatform)
 {
+	using namespace UE::NNERuntimeORT::Private;
+
 	if (CanCreateModelData(FileType, FileData, AdditionalFileData, FileId, TargetPlatform) != ECanCreateModelDataStatus::Ok)
 	{
 		UE_LOG(LogNNE, Warning, TEXT("UNNERuntimeORTCpu cannot create the model data with id %s (Filetype: %s)"), *FileId.ToString(EGuidFormats::Digits).ToLower(), *FileType);
@@ -38,9 +40,16 @@ TSharedPtr<UE::NNE::FSharedModelData> UNNERuntimeORTCpu::CreateModelData(const F
 	}
 
 	FNNEModelRaw InputModel{TArray<uint8>{FileData}, ENNEInferenceFormat::ONNX};
-	if (!UE::NNERuntimeORT::Private::OrtHelper::OptimizeModel(Environment.ToSharedRef(), InputModel, ENNEInferenceFormat::ORT))
+	if (GraphOptimizationLevel OptimizationLevel = GetGraphOptimizationLevelForCPU(false, IsRunningCookCommandlet()); OptimizationLevel > GraphOptimizationLevel::ORT_DISABLE_ALL)
 	{
-		return {};
+		TUniquePtr<Ort::SessionOptions> SessionOptions = CreateSessionOptionsDefault(Environment.ToSharedRef());
+		SessionOptions->SetGraphOptimizationLevel(OptimizationLevel);
+		SessionOptions->EnableCpuMemArena();
+
+		if (!OptimizeModel(Environment.ToSharedRef(), *SessionOptions, ENNEInferenceFormat::ONNX, InputModel))
+		{
+			return {};
+		}
 	}
 
 	TArray<uint8> Result;
@@ -142,6 +151,8 @@ UNNERuntimeORTDml::ECanCreateModelDataStatus UNNERuntimeORTDml::CanCreateModelDa
 
 TSharedPtr<UE::NNE::FSharedModelData> UNNERuntimeORTDml::CreateModelData(const FString& FileType, TConstArrayView<uint8> FileData, const TMap<FString, TConstArrayView<uint8>>& AdditionalFileData, const FGuid& FileId, const ITargetPlatform* TargetPlatform)
 {
+	using namespace UE::NNERuntimeORT::Private;
+
 	if (CanCreateModelData(FileType, FileData, AdditionalFileData, FileId, TargetPlatform) != ECanCreateModelDataStatus::Ok)
 	{
 		UE_LOG(LogNNE, Warning, TEXT("NNERuntimeORTDmlImpl cannot create the model data with id %s (Filetype: %s)"), *FileId.ToString(EGuidFormats::Digits).ToLower(), *FileType);
@@ -149,9 +160,17 @@ TSharedPtr<UE::NNE::FSharedModelData> UNNERuntimeORTDml::CreateModelData(const F
 	}
 
 	FNNEModelRaw InputModel{TArray<uint8>{FileData}, ENNEInferenceFormat::ONNX};
-	if (!UE::NNERuntimeORT::Private::OrtHelper::OptimizeModel(Environment.ToSharedRef(), InputModel, ENNEInferenceFormat::ONNX))
+	if (GraphOptimizationLevel OptimizationLevel = GetGraphOptimizationLevelForDML(false, IsRunningCookCommandlet()); OptimizationLevel > GraphOptimizationLevel::ORT_DISABLE_ALL)
 	{
-		return {};
+		TUniquePtr<Ort::SessionOptions> SessionOptions = CreateSessionOptionsDefault(Environment.ToSharedRef());
+		SessionOptions->SetGraphOptimizationLevel(OptimizationLevel);
+		SessionOptions->SetExecutionMode(ExecutionMode::ORT_SEQUENTIAL);
+		SessionOptions->DisableMemPattern();
+		
+		if (!OptimizeModel(Environment.ToSharedRef(), *SessionOptions, ENNEInferenceFormat::ONNX, InputModel))
+		{
+			return {};
+		}
 	}
 
 	TArray<uint8> Result;
diff --git a/Engine/Plugins/NNE/NNERuntimeORT/Source/NNERuntimeORT/Private/NNERuntimeORTModel.cpp b/Engine/Plugins/NNE/NNERuntimeORT/Source/NNERuntimeORT/Private/NNERuntimeORTModel.cpp
--- a/Engine/Plugins/NNE/NNERuntimeORT/Source/NNERuntimeORT/Private/NNERuntimeORTModel.cpp
+++ b/Engine/Plugins/NNE/NNERuntimeORT/Source/NNERuntimeORT/Private/NNERuntimeORTModel.cpp
@@ -25,16 +25,6 @@ BEGIN_SHADER_PARAMETER_STRUCT(FORTModelInstanceRDGParameters, )
 	RDG_BUFFER_ACCESS_ARRAY(OutputBuffers)
 END_SHADER_PARAMETER_STRUCT()
 
-static int32 ORTProfilingSessionNumber = 0;
-static TAutoConsoleVariable<bool> CVarNNERuntimeORTEnableProfiling(
-	TEXT("nne.ort.enableprofiling"),
-	false,
-	TEXT("True if NNERuntimeORT plugin should create ORT sessions with profiling enabled.\n")
-	TEXT("When profiling is enabled ORT will create standard performance tracing json files next to the editor executable.\n")
-	TEXT("The files will be prefixed by 'NNERuntimeORTProfile_' and can be loaded for example using chrome://tracing.\n")
-	TEXT("More information can be found at https://onnxruntime.ai/docs/performance/tune-performance/profiling-tools.html\n"),
-	ECVF_Default);
-
 DECLARE_GPU_STAT_NAMED(FNNERuntimeORTDmlRDG, TEXT("FModelInstanceORTDmlRDG::EnqueueRDG"));
 
 namespace UE::NNERuntimeORT::Private
@@ -46,149 +36,16 @@ namespace Detail
 	FRuntimeConf MakeRuntimeConfigFromSettings(const UNNERuntimeORTSettings* Settings)
 	{
 		FRuntimeConf Result{};
-
 #if WITH_EDITOR
 			FThreadingOptions ThreadingOptions = Settings->EditorThreadingOptions;
 #else
 			FThreadingOptions ThreadingOptions = Settings->GameThreadingOptions;
 #endif
-
-		Result.bUseGlobalThreadPool = ThreadingOptions.bUseGlobalThreadPool;
-		Result.IntraOpNumThreads = ThreadingOptions.IntraOpNumThreads;
-		Result.InterOpNumThreads = ThreadingOptions.InterOpNumThreads;
 		Result.ExecutionMode = ThreadingOptions.ExecutionMode == EExecutionMode::SEQUENTIAL ? ExecutionMode::ORT_SEQUENTIAL : ExecutionMode::ORT_PARALLEL;
 
 		return Result;
 	}
 
-TUniquePtr<Ort::SessionOptions> CreateSessionOptionsDefault(const FRuntimeConf &RuntimeConf)
-{
-	TUniquePtr<Ort::SessionOptions> SessionOptions = MakeUnique<Ort::SessionOptions>();
-
-	// Configure Threading
-	if (RuntimeConf.bUseGlobalThreadPool)
-	{
-		SessionOptions->DisablePerSessionThreads();
-	}
-	else
-	{
-		SessionOptions->SetIntraOpNumThreads(RuntimeConf.IntraOpNumThreads);
-		SessionOptions->SetInterOpNumThreads(RuntimeConf.InterOpNumThreads);
-	}
-
-	SessionOptions->SetExecutionMode(RuntimeConf.ExecutionMode);
-
-	// Configure Graph optimization
-	SessionOptions->SetGraphOptimizationLevel(RuntimeConf.OptimizationLevel);
-
-
-	// Configure Profiling
-	if (CVarNNERuntimeORTEnableProfiling.GetValueOnGameThread())
-	{
-		FString ProfilingFilePrefix("NNERuntimeORTProfile_");
-		ProfilingFilePrefix += FString::FromInt(ORTProfilingSessionNumber);
-		++ORTProfilingSessionNumber;
-		#if PLATFORM_WINDOWS
-			SessionOptions->EnableProfiling(*ProfilingFilePrefix);
-		#else
-			SessionOptions->EnableProfiling(TCHAR_TO_ANSI(*ProfilingFilePrefix));
-		#endif
-	}
-
-	return SessionOptions;
-}
-
-#if PLATFORM_WINDOWS
-TUniquePtr<Ort::SessionOptions> CreateSessionOptionsForDirectML(const FRuntimeConf &RuntimeConf)
-{
-	TUniquePtr<Ort::SessionOptions> SessionOptions = CreateSessionOptionsDefault(RuntimeConf);
-	if (!SessionOptions.IsValid())
-	{
-		return {};
-	}
-
-	// Configure for DirectML
-	SessionOptions->SetExecutionMode(ExecutionMode::ORT_SEQUENTIAL);
-	SessionOptions->DisableMemPattern();
-
-	// In order to use DirectML we need D3D12
-	ID3D12DynamicRHI* RHI = nullptr;
-
-	if (!GDynamicRHI)
-	{
-		UE_LOG(LogNNE, Error, TEXT("Error:No RHI found, could not initialize"));
-		return {};
-	}
-
-	if (IsRHID3D12() )
-	{
-		RHI = GetID3D12DynamicRHI();
-	}
-	else
-	{
-		if (GDynamicRHI)
-		{
-			UE_LOG(LogNNE, Error, TEXT("Error:%s RHI is not supported by DirectML, please use D3D12."), GDynamicRHI->GetName());
-			return {};
-		}
-		else
-		{
-			UE_LOG(LogNNE, Error, TEXT("Error:No RHI found"));
-			return {};
-		}
-	}
-
-	check(RHI);
-
-	const int32 DeviceIndex = 0;
-	ID3D12Device* D3D12Device = RHI->RHIGetDevice(DeviceIndex);
-
-	if (!D3D12Device)
-	{
-		UE_LOG(LogNNE, Error, TEXT("Failed to get D3D12 Device from RHI for device index %d"), DeviceIndex);
-		return {};
-	}
-
-	DML_CREATE_DEVICE_FLAGS DmlCreateFlags = DML_CREATE_DEVICE_FLAG_NONE;
-
-	// Set debugging flags
-	if (GRHIGlobals.IsDebugLayerEnabled)
-	{
-		DmlCreateFlags |= DML_CREATE_DEVICE_FLAG_DEBUG;
-	}
-
-	IDMLDevice* DmlDevice = nullptr;
-	HRESULT Res = DMLCreateDevice(D3D12Device, DmlCreateFlags, DML_PPV_ARGS(&DmlDevice));
-
-	if (FAILED(Res) || !DmlDevice)
-	{
-		UE_LOG(LogNNE, Error, TEXT("Failed to create DirectML device, DMLCreateDevice error code :%x"), Res);
-		return {};
-	}
-
-	ID3D12CommandQueue* CmdQ = RHI->RHIGetCommandQueue();
-
-	const OrtDmlApi* DmlApi = nullptr;
-	Ort::ThrowOnError(Ort::GetApi().GetExecutionProviderApi("DML", ORT_API_VERSION, reinterpret_cast<const void**>(&DmlApi)));
-
-	if (!DmlApi)
-	{
-		UE_LOG(LogNNE, Error, TEXT("Ort DirectML Api not available!"));
-		return {};
-	}
-
-	OrtStatusPtr Status = DmlApi->SessionOptionsAppendExecutionProvider_DML1(*SessionOptions.Get(), DmlDevice, CmdQ);
-
-	if (Status)
-	{
-		UE_LOG(LogNNE, Error, TEXT("Failed to add DirectML execution provider to OnnxRuntime session options: %s"), ANSI_TO_TCHAR(Ort::GetApi().GetErrorMessage(Status)));
-		return {};
-	}
-
-	return SessionOptions;
-}
-#endif // PLATFORM_WINDOWS
-
 } // namespace Detail
 
 template <class ModelInterface, class TensorBinding> 
@@ -508,12 +365,14 @@ bool FModelInstanceORTCpu::InitializedAndConfigureMembers()
 		return false;
 	}
 
-	SessionOptions = Detail::CreateSessionOptionsDefault(RuntimeConf);
+	SessionOptions = CreateSessionOptionsDefault(Environment);
 	if (!SessionOptions.IsValid())
 	{
 		return false;
 	}
 
+	SessionOptions->SetExecutionMode(RuntimeConf.ExecutionMode);
+	SessionOptions->SetGraphOptimizationLevel(GetGraphOptimizationLevelForCPU(true));
 	SessionOptions->EnableCpuMemArena();
 
 	return true;
@@ -549,12 +408,14 @@ bool FModelInstanceORTDmlGPU::InitializedAndConfigureMembers()
 		return false;
 	}
 
-	SessionOptions = Detail::CreateSessionOptionsForDirectML(RuntimeConf);
+	SessionOptions = CreateSessionOptionsForDirectML(Environment);
 	if (!SessionOptions.IsValid())
 	{
 		return false;
 	}
 
+	SessionOptions->SetGraphOptimizationLevel(GetGraphOptimizationLevelForDML(true));
+
 	return true;
 }
 
@@ -598,12 +459,14 @@ bool FModelInstanceORTDmlRDG::Init()
 	{
 		Allocator = MakeUnique<Ort::AllocatorWithDefaultOptions>();
 
-		SessionOptions = Detail::CreateSessionOptionsForDirectML(RuntimeConf);
+		SessionOptions = CreateSessionOptionsForDirectML(Environment);
 		if (!SessionOptions.IsValid())
 		{
 			UE_LOG(LogNNE, Error, TEXT("FModelInstanceORTDmlRDG::Init(): Failed to configure session options for DirectML Execution Provider."));
 			return false;
 		}
+		
+		SessionOptions->SetGraphOptimizationLevel(GetGraphOptimizationLevelForDML(true));
 
 		Session = MakeUnique<Ort::Session>(Environment->GetOrtEnv(), ModelBuffer.GetData(), ModelBuffer.Num(), *SessionOptions);
 
@@ -763,13 +626,15 @@ FModelInstanceORTDmlRDG::ESetInputTensorShapesStatus FModelInstanceORTDmlRDG::Se
 	}
 
 	// Recreate session options because potentially we add new free dimension overrides
-	SessionOptions = Detail::CreateSessionOptionsForDirectML(RuntimeConf);
+	SessionOptions = CreateSessionOptionsForDirectML(Environment);
 	if (!SessionOptions.IsValid())
 	{
 		UE_LOG(LogNNE, Error, TEXT("Failed to recreate session options!"));
 		return ESetInputTensorShapesStatus::Fail;
 	}
 
+	SessionOptions->SetGraphOptimizationLevel(GetGraphOptimizationLevelForDML(true));
+
 	// Setup concrete input tensors
 	for (int32 i = 0; i < InputSymbolicTensors.Num(); i++)
 	{
diff --git a/Engine/Plugins/NNE/NNERuntimeORT/Source/NNERuntimeORT/Private/NNERuntimeORTModel.h b/Engine/Plugins/NNE/NNERuntimeORT/Source/NNERuntimeORT/Private/NNERuntimeORTModel.h
--- a/Engine/Plugins/NNE/NNERuntimeORT/Source/NNERuntimeORT/Private/NNERuntimeORTModel.h
+++ b/Engine/Plugins/NNE/NNERuntimeORT/Source/NNERuntimeORT/Private/NNERuntimeORTModel.h
@@ -18,13 +18,7 @@ namespace UE::NNERuntimeORT::Private
 
 struct FRuntimeConf
 {
-	bool bUseGlobalThreadPool = false;
-	int32 IntraOpNumThreads = 0;
-	int32 InterOpNumThreads = 0;
-
 	ExecutionMode ExecutionMode = ExecutionMode::ORT_SEQUENTIAL;
-
-	GraphOptimizationLevel OptimizationLevel = GraphOptimizationLevel::ORT_ENABLE_ALL;
 };
 
 template <class ModelInterface, class TensorBinding>
diff --git a/Engine/Plugins/NNE/NNERuntimeORT/Source/NNERuntimeORT/Private/NNERuntimeORTUtils.cpp b/Engine/Plugins/NNE/NNERuntimeORT/Source/NNERuntimeORT/Private/NNERuntimeORTUtils.cpp
--- a/Engine/Plugins/NNE/NNERuntimeORT/Source/NNERuntimeORT/Private/NNERuntimeORTUtils.cpp
+++ b/Engine/Plugins/NNE/NNERuntimeORT/Source/NNERuntimeORT/Private/NNERuntimeORTUtils.cpp
@@ -8,11 +8,81 @@
 #include "NNE.h"
 #include "NNERuntimeORTEnv.h"
 
+#if PLATFORM_WINDOWS
+#include "ID3D12DynamicRHI.h"
+#endif // PLATFORM_WINDOWS
+
+// DirectML is implemented using COM on all platforms
+#ifdef IID_GRAPHICS_PPV_ARGS
+#define DML_PPV_ARGS(x) __uuidof(*x), IID_PPV_ARGS_Helper(x)
+#else
+#define DML_PPV_ARGS(x) IID_PPV_ARGS(x)
+#endif
+
+static int32 ORTProfilingSessionNumber = 0;
+static TAutoConsoleVariable<bool> CVarNNERuntimeORTEnableProfiling(
+	TEXT("nne.ort.enableprofiling"),
+	false,
+	TEXT("True if NNERuntimeORT plugin should create ORT sessions with profiling enabled.\n")
+	TEXT("When profiling is enabled ORT will create standard performance tracing json files next to the editor executable.\n")
+	TEXT("The files will be prefixed by 'NNERuntimeORTProfile_' and can be loaded for example using chrome://tracing.\n")
+	TEXT("More information can be found at https://onnxruntime.ai/docs/performance/tune-performance/profiling-tools.html\n"),
+	ECVF_Default);
+
 namespace UE::NNERuntimeORT::Private
 {
-namespace OrtHelper
+// For more details about ORT graph optimization checkout
+// https://onnxruntime.ai/docs/performance/model-optimizations/graph-optimizations.html
+
+struct FGraphOptimizationLevels
+{
+	GraphOptimizationLevel Cooking;
+	GraphOptimizationLevel Offline;
+	GraphOptimizationLevel Online;
+};
+
+// CPU
+static constexpr FGraphOptimizationLevels OrtCpuOptimizationLevels
 {
+	.Cooking = GraphOptimizationLevel::ORT_ENABLE_EXTENDED,
+	.Offline = GraphOptimizationLevel::ORT_DISABLE_ALL,
+	.Online = GraphOptimizationLevel::ORT_ENABLE_ALL
+};
+
+// DirectML EP
+// note: optimize with DirectML EP enabled, but currently an offline optimized model can not be optimized again (only DML)!
+// Therefore, if one enables offline optimization, set it to ORT_ENABLE_ALL and disable any optimization in online mode (ORT_DISABLE_ALL).
+//
+// note: since cooked models contain only basic graph optimizations, we need full optimization in online mode.
+// Therefore, offline optimization in non-Editor can not be turned on.
+static constexpr FGraphOptimizationLevels OrtDmlOptimizationLevels
+{
+	.Cooking = GraphOptimizationLevel::ORT_ENABLE_BASIC,
+	.Offline = GraphOptimizationLevel::ORT_DISABLE_ALL,
+	.Online = GraphOptimizationLevel::ORT_ENABLE_ALL
+};
 
+GraphOptimizationLevel GetGraphOptimizationLevel(const FGraphOptimizationLevels &OptimizationLevels, bool bIsOnline, bool bIsCooking)
+{
+	if (bIsOnline)
+	{
+		return OptimizationLevels.Online;
+	}
+	else
+	{
+		if (bIsCooking)
+		{
+			return OptimizationLevels.Cooking;
+		}
+		else
+		{
+			return OptimizationLevels.Offline;
+		}
+	}
+}
+
+namespace OrtHelper
+{
 TArray<uint32> GetShape(const Ort::Value& OrtTensor)
 {
 	OrtTensorTypeAndShapeInfo* TypeAndShapeInfoPtr = nullptr;
@@ -37,8 +107,143 @@ TArray<uint32> GetShape(const Ort::Value& OrtTensor)
 
 	return Result;
 }
+} // namespace OrtHelper
 
-bool OptimizeModel(TSharedRef<FEnvironment> InEnvironment, FNNEModelRaw& Model, ENNEInferenceFormat TargetFormat)
+GraphOptimizationLevel GetGraphOptimizationLevelForCPU(bool bIsOnline, bool bIsCooking)
+{
+	return GetGraphOptimizationLevel(OrtCpuOptimizationLevels, bIsOnline, bIsCooking);
+}
+
+GraphOptimizationLevel GetGraphOptimizationLevelForDML(bool bIsOnline, bool bIsCooking)
+{
+	return GetGraphOptimizationLevel(OrtDmlOptimizationLevels, bIsOnline, bIsCooking);
+}
+
+TUniquePtr<Ort::SessionOptions> CreateSessionOptionsDefault(const TSharedRef<FEnvironment> &Environment)
+{
+	const FEnvironment::FConfig Config = Environment->GetConfig();
+
+	TUniquePtr<Ort::SessionOptions> SessionOptions = MakeUnique<Ort::SessionOptions>();
+
+	// Configure Threading
+	if (Config.bUseGlobalThreadPool)
+	{
+		SessionOptions->DisablePerSessionThreads();
+	}
+	else
+	{
+		SessionOptions->SetIntraOpNumThreads(Config.IntraOpNumThreads);
+		SessionOptions->SetInterOpNumThreads(Config.InterOpNumThreads);
+	}
+
+	// Configure Profiling
+	if (CVarNNERuntimeORTEnableProfiling.GetValueOnGameThread())
+	{
+		FString ProfilingFilePrefix("NNERuntimeORTProfile_");
+		ProfilingFilePrefix += FString::FromInt(ORTProfilingSessionNumber);
+		++ORTProfilingSessionNumber;
+		#if PLATFORM_WINDOWS
+			SessionOptions->EnableProfiling(*ProfilingFilePrefix);
+		#else
+			SessionOptions->EnableProfiling(TCHAR_TO_ANSI(*ProfilingFilePrefix));
+		#endif
+	}
+
+	return SessionOptions;
+}
+
+#if PLATFORM_WINDOWS
+TUniquePtr<Ort::SessionOptions> CreateSessionOptionsForDirectML(const TSharedRef<FEnvironment> &Environment)
+{
+	TUniquePtr<Ort::SessionOptions> SessionOptions = CreateSessionOptionsDefault(Environment);
+	if (!SessionOptions.IsValid())
+	{
+		return {};
+	}
+
+	// Configure for DirectML
+	SessionOptions->SetExecutionMode(ExecutionMode::ORT_SEQUENTIAL);
+	SessionOptions->DisableMemPattern();
+
+	// In order to use DirectML we need D3D12
+	ID3D12DynamicRHI* RHI = nullptr;
+
+	if (!GDynamicRHI)
+	{
+		UE_LOG(LogNNE, Error, TEXT("Error:No RHI found, could not initialize"));
+		return {};
+	}
+
+	if (IsRHID3D12() )
+	{
+		RHI = GetID3D12DynamicRHI();
+	}
+	else
+	{
+		if (GDynamicRHI)
+		{
+			UE_LOG(LogNNE, Error, TEXT("Error:%s RHI is not supported by DirectML, please use D3D12."), GDynamicRHI->GetName());
+			return {};
+		}
+		else
+		{
+			UE_LOG(LogNNE, Error, TEXT("Error:No RHI found"));
+			return {};
+		}
+	}
+
+	check(RHI);
+
+	const int32 DeviceIndex = 0;
+	ID3D12Device* D3D12Device = RHI->RHIGetDevice(DeviceIndex);
+
+	if (!D3D12Device)
+	{
+		UE_LOG(LogNNE, Error, TEXT("Failed to get D3D12 Device from RHI for device index %d"), DeviceIndex);
+		return {};
+	}
+
+	DML_CREATE_DEVICE_FLAGS DmlCreateFlags = DML_CREATE_DEVICE_FLAG_NONE;
+
+	// Set debugging flags
+	if (GRHIGlobals.IsDebugLayerEnabled)
+	{
+		DmlCreateFlags |= DML_CREATE_DEVICE_FLAG_DEBUG;
+	}
+
+	IDMLDevice* DmlDevice = nullptr;
+	HRESULT Res = DMLCreateDevice(D3D12Device, DmlCreateFlags, DML_PPV_ARGS(&DmlDevice));
+
+	if (FAILED(Res) || !DmlDevice)
+	{
+		UE_LOG(LogNNE, Error, TEXT("Failed to create DirectML device, DMLCreateDevice error code :%x"), Res);
+		return {};
+	}
+
+	ID3D12CommandQueue* CmdQ = RHI->RHIGetCommandQueue();
+
+	const OrtDmlApi* DmlApi = nullptr;
+	Ort::ThrowOnError(Ort::GetApi().GetExecutionProviderApi("DML", ORT_API_VERSION, reinterpret_cast<const void**>(&DmlApi)));
+
+	if (!DmlApi)
+	{
+		UE_LOG(LogNNE, Error, TEXT("Ort DirectML Api not available!"));
+		return {};
+	}
+
+	OrtStatusPtr Status = DmlApi->SessionOptionsAppendExecutionProvider_DML1(*SessionOptions, DmlDevice, CmdQ);
+
+	if (Status)
+	{
+		UE_LOG(LogNNE, Error, TEXT("Failed to add DirectML execution provider to OnnxRuntime session options: %s"), ANSI_TO_TCHAR(Ort::GetApi().GetErrorMessage(Status)));
+		return {};
+	}
+
+	return SessionOptions;
+}
+#endif // PLATFORM_WINDOWS
+
+bool OptimizeModel(const TSharedRef<FEnvironment> &Environment, Ort::SessionOptions &SessionOptions, ENNEInferenceFormat TargetFormat, FNNEModelRaw& Model)
 {
 	SCOPED_NAMED_EVENT_TEXT("OrtHelper::OptimizeModel", FColor::Magenta);
 
@@ -49,39 +254,21 @@ bool OptimizeModel(TSharedRef<FEnvironment> InEnvironment, FNNEModelRaw& Model,
 	}
 
 	FString ProjIntermediateDir = FPaths::ConvertRelativePathToFull(FPaths::ProjectIntermediateDir());
-	FString ModelToOptimizePath = FPaths::CreateTempFilename(*ProjIntermediateDir, TEXT("ORTOptimizerPass_ToOptimize"), TEXT(".onnx"));
 	FString TargetExtension = TargetFormat == ENNEInferenceFormat::ONNX ? TEXT(".onnx") : TEXT(".ort");
 	FString ModelOptimizedPath = FPaths::CreateTempFilename(*ProjIntermediateDir, TEXT("ORTOptimizerPass_Optimized"), *TargetExtension);
 
-	//See https://onnxruntime.ai/docs/performance/model-optimizations/graph-optimizations.html
-	//We only enable all the optimization when going to ORT format itself for the CPU provider
-	GraphOptimizationLevel OptimizationLevel = TargetFormat == ENNEInferenceFormat::ONNX ? ORT_ENABLE_BASIC : ORT_ENABLE_ALL;
-
-	FFileHelper::SaveArrayToFile(Model.Data, *ModelToOptimizePath);
-
 #if WITH_EDITOR
 	try
 #endif // WITH_EDITOR
 	{
-		Ort::SessionOptions SessionOptions;
-		if (ensureMsgf(InEnvironment->GetConfig().bUseGlobalThreadPool, TEXT("For Model Optimizer use ONNX Runtime global threadpool to improve performance!")))
-		{
-			SessionOptions.DisablePerSessionThreads();
-		}
-		else
-		{
-			SessionOptions.SetIntraOpNumThreads(InEnvironment->GetConfig().IntraOpNumThreads);
-			SessionOptions.SetInterOpNumThreads(InEnvironment->GetConfig().InterOpNumThreads);
-		}
-		SessionOptions.SetGraphOptimizationLevel(OptimizationLevel);
 #if PLATFORM_WINDOWS
 		SessionOptions.SetOptimizedModelFilePath(*ModelOptimizedPath);
 
-		Ort::Session Session(InEnvironment->GetOrtEnv(), *ModelToOptimizePath, SessionOptions);
+		Ort::Session Session(Environment->GetOrtEnv(), Model.Data.GetData(), Model.Data.Num(), SessionOptions);
 #else
 		SessionOptions.SetOptimizedModelFilePath(TCHAR_TO_ANSI(*ModelOptimizedPath));
-		
-		Ort::Session Session(InEnvironment->GetOrtEnv(), TCHAR_TO_ANSI(*ModelToOptimizePath), SessionOptions);
+
+		Ort::Session Session(Environment->GetOrtEnv(), Model.Data.GetData(), Model.Data.Num(), SessionOptions);
 #endif
 	}
 #if WITH_EDITOR
@@ -99,7 +286,6 @@ bool OptimizeModel(TSharedRef<FEnvironment> InEnvironment, FNNEModelRaw& Model,
 
 	FFileHelper::LoadFileToArray(Model.Data, *ModelOptimizedPath);
 
-	IFileManager::Get().Delete(*ModelToOptimizePath);
 	IFileManager::Get().Delete(*ModelOptimizedPath);
 
 	Model.Format = TargetFormat;
@@ -107,8 +293,6 @@ bool OptimizeModel(TSharedRef<FEnvironment> InEnvironment, FNNEModelRaw& Model,
 	return true;
 }
 
-} // OrtHelper
-
 TypeInfoORT TranslateTensorTypeORTToNNE(ONNXTensorElementDataType OrtDataType)
 {
 	ENNETensorDataType DataType = ENNETensorDataType::None;
diff --git a/Engine/Plugins/NNE/NNERuntimeORT/Source/NNERuntimeORT/Private/NNERuntimeORTUtils.h b/Engine/Plugins/NNE/NNERuntimeORT/Source/NNERuntimeORT/Private/NNERuntimeORTUtils.h
--- a/Engine/Plugins/NNE/NNERuntimeORT/Source/NNERuntimeORT/Private/NNERuntimeORTUtils.h
+++ b/Engine/Plugins/NNE/NNERuntimeORT/Source/NNERuntimeORT/Private/NNERuntimeORTUtils.h
@@ -6,18 +6,29 @@
 #include "NNEOnnxruntime.h"
 #include "NNERuntimeFormat.h"
 #include "NNETypes.h"
+#include "Templates/UniquePtr.h"
 
 namespace UE::NNERuntimeORT::Private
 {
 
 	class FEnvironment;
 
-	namespace OrtHelper
-	{
-		TArray<uint32> GetShape(const Ort::Value& OrtTensor);
+namespace OrtHelper
+{
+	TArray<uint32> GetShape(const Ort::Value& OrtTensor);
+} // namespace OrtHelper
+
+	GraphOptimizationLevel GetGraphOptimizationLevelForCPU(bool bIsOnline, bool bIsCooking = false);
+
+	GraphOptimizationLevel GetGraphOptimizationLevelForDML(bool bIsOnline, bool bIsCooking = false);
+
+	TUniquePtr<Ort::SessionOptions> CreateSessionOptionsDefault(const TSharedRef<FEnvironment> &Environment);
+
+#if PLATFORM_WINDOWS
+	TUniquePtr<Ort::SessionOptions> CreateSessionOptionsForDirectML(const TSharedRef<FEnvironment> &Environment);
+#endif // PLATFORM_WINDOWS
 
-		bool OptimizeModel(TSharedRef<FEnvironment> InEnvironment, FNNEModelRaw& Model, ENNEInferenceFormat OutFormat);
-	}
+	bool OptimizeModel(const TSharedRef<FEnvironment> &Environment, Ort::SessionOptions &SessionOptions, ENNEInferenceFormat TargetFormat, FNNEModelRaw& Model);
 
 	struct TypeInfoORT
 	{
