[NNE] NNERuntimeORTDml (GPU) also available in head-less mode.

#tests Editor with nullrhi and vulkan run tests with another instance using session frontend

#rnx

diff --git a/Engine/Plugins/NNE/NNERuntimeORT/Source/NNERuntimeORT/Private/NNERuntimeORT.cpp b/Engine/Plugins/NNE/NNERuntimeORT/Source/NNERuntimeORT/Private/NNERuntimeORT.cpp
--- a/Engine/Plugins/NNE/NNERuntimeORT/Source/NNERuntimeORT/Private/NNERuntimeORT.cpp
+++ b/Engine/Plugins/NNE/NNERuntimeORT/Source/NNERuntimeORT/Private/NNERuntimeORT.cpp
@@ -134,9 +134,11 @@ TSharedPtr<UE::NNE::IModelCPU> UNNERuntimeORTCpu::CreateModelCPU(const TObjectPt
 /*
  * UNNERuntimeORTDml
  */
-void UNNERuntimeORTDml::Init(TSharedRef<UE::NNERuntimeORT::Private::FEnvironment> InEnvironment)
+void UNNERuntimeORTDml::Init(TSharedRef<UE::NNERuntimeORT::Private::FEnvironment> InEnvironment, bool bInDirectMLAvailable)
 {
 	Environment = InEnvironment;
+	bDirectMLAvailable = bInDirectMLAvailable;
+	bD3D12Available = UE::NNERuntimeORT::Private::IsD3D12Available();
 }
 
 FString UNNERuntimeORTDml::GetRuntimeName() const
@@ -189,7 +191,7 @@ FString UNNERuntimeORTDml::GetModelDataIdentifier(const FString& FileType, TCons
 
 UNNERuntimeORTDml::ECanCreateModelGPUStatus UNNERuntimeORTDml::CanCreateModelGPU(const TObjectPtr<UNNEModelData> ModelData) const
 {
-	return CanCreateModelCommon(ModelData) == ECanCreateModelCommonStatus::Ok ? ECanCreateModelGPUStatus::Ok : ECanCreateModelGPUStatus::Fail;
+	return CanCreateModelCommon(ModelData, false) == ECanCreateModelCommonStatus::Ok ? ECanCreateModelGPUStatus::Ok : ECanCreateModelGPUStatus::Fail;
 }
 
 TSharedPtr<UE::NNE::IModelGPU> UNNERuntimeORTDml::CreateModelGPU(const TObjectPtr<UNNEModelData> ModelData)
@@ -256,13 +258,23 @@ TSharedPtr<UE::NNE::IModelRDG> UNNERuntimeORTDml::CreateModelRDG(TObjectPtr<UNNE
 #endif // PLATFORM_WINDOWS
 }
 
-UNNERuntimeORTDml::ECanCreateModelCommonStatus UNNERuntimeORTDml::CanCreateModelCommon(const TObjectPtr<UNNEModelData> ModelData) const
+UNNERuntimeORTDml::ECanCreateModelCommonStatus UNNERuntimeORTDml::CanCreateModelCommon(const TObjectPtr<UNNEModelData> ModelData, bool bRHID3D12Required) const
 {
 #if PLATFORM_WINDOWS
 	check(ModelData != nullptr);
 
-	// In order to use DirectML we need D3D12
-	if (!IsRHID3D12())
+	// DirectML is required
+	if (!bDirectMLAvailable)
+	{
+		return ECanCreateModelCommonStatus::Fail;
+	}
+
+	// Either RHID3D12 or at least D3D12 is required
+	if (bRHID3D12Required && !IsRHID3D12())
+	{
+		return ECanCreateModelCommonStatus::Fail;
+	}
+	else if (!bD3D12Available)
 	{
 		return ECanCreateModelCommonStatus::Fail;
 	}
diff --git a/Engine/Plugins/NNE/NNERuntimeORT/Source/NNERuntimeORT/Private/NNERuntimeORT.h b/Engine/Plugins/NNE/NNERuntimeORT/Source/NNERuntimeORT/Private/NNERuntimeORT.h
--- a/Engine/Plugins/NNE/NNERuntimeORT/Source/NNERuntimeORT/Private/NNERuntimeORT.h
+++ b/Engine/Plugins/NNE/NNERuntimeORT/Source/NNERuntimeORT/Private/NNERuntimeORT.h
@@ -59,7 +59,7 @@ public:
 	
 	virtual ~UNNERuntimeORTDml() = default;
 
-	void Init(TSharedRef<UE::NNERuntimeORT::Private::FEnvironment> InEnvironment);
+	void Init(TSharedRef<UE::NNERuntimeORT::Private::FEnvironment> InEnvironment, bool bInDirectMLAvailable);
 
 	virtual FString GetRuntimeName() const override;
 
@@ -74,5 +74,8 @@ public:
 	virtual TSharedPtr<UE::NNE::IModelRDG> CreateModelRDG(TObjectPtr<UNNEModelData> ModelData) override;
 
 private:
-	ECanCreateModelCommonStatus CanCreateModelCommon(const TObjectPtr<UNNEModelData> ModelData) const;
+	ECanCreateModelCommonStatus CanCreateModelCommon(const TObjectPtr<UNNEModelData> ModelData, bool bRHID3D12Required = true) const;
+
+	bool bDirectMLAvailable = false;
+	bool bD3D12Available = false;
 };
\ No newline at end of file
diff --git a/Engine/Plugins/NNE/NNERuntimeORT/Source/NNERuntimeORT/Private/NNERuntimeORTModel.cpp b/Engine/Plugins/NNE/NNERuntimeORT/Source/NNERuntimeORT/Private/NNERuntimeORTModel.cpp
--- a/Engine/Plugins/NNE/NNERuntimeORT/Source/NNERuntimeORT/Private/NNERuntimeORTModel.cpp
+++ b/Engine/Plugins/NNE/NNERuntimeORT/Source/NNERuntimeORT/Private/NNERuntimeORTModel.cpp
@@ -407,7 +407,7 @@ bool FModelInstanceORTDmlGPU::InitializedAndConfigureMembers()
 		return false;
 	}
 
-	SessionOptions = CreateSessionOptionsForDirectML(Environment);
+	SessionOptions = CreateSessionOptionsForDirectML(Environment, false);
 	if (!SessionOptions.IsValid())
 	{
 		return false;
diff --git a/Engine/Plugins/NNE/NNERuntimeORT/Source/NNERuntimeORT/Private/NNERuntimeORTModule.cpp b/Engine/Plugins/NNE/NNERuntimeORT/Source/NNERuntimeORT/Private/NNERuntimeORTModule.cpp
--- a/Engine/Plugins/NNE/NNERuntimeORT/Source/NNERuntimeORT/Private/NNERuntimeORTModule.cpp
+++ b/Engine/Plugins/NNE/NNERuntimeORT/Source/NNERuntimeORT/Private/NNERuntimeORTModule.cpp
@@ -99,18 +99,15 @@ void FNNERuntimeORTModule::StartupModule()
 	EnvironmentHelper::CreateOrtEnvFromSettings(GetDefault<UNNERuntimeORTSettings>(), *Environment);
 
 #if PLATFORM_WINDOWS
-	if (bDirectMLDllLoaded)
+	// NNE runtime ORT Dml startup
+	NNERuntimeORTDml = NewObject<UNNERuntimeORTDml>();
+	if (NNERuntimeORTDml.IsValid())
 	{
-		// NNE runtime ORT Dml startup
-		NNERuntimeORTDml = NewObject<UNNERuntimeORTDml>();
-		if (NNERuntimeORTDml.IsValid())
-		{
-			TWeakInterfacePtr<INNERuntime> RuntimeDmlInterface(NNERuntimeORTDml.Get());
+		TWeakInterfacePtr<INNERuntime> RuntimeDmlInterface(NNERuntimeORTDml.Get());
 
-			NNERuntimeORTDml->Init(Environment.ToSharedRef());
-			NNERuntimeORTDml->AddToRoot();
-			UE::NNE::RegisterRuntime(RuntimeDmlInterface);
-		}
+		NNERuntimeORTDml->Init(Environment.ToSharedRef(), bDirectMLDllLoaded);
+		NNERuntimeORTDml->AddToRoot();
+		UE::NNE::RegisterRuntime(RuntimeDmlInterface);
 	}
 #endif // PLATFORM_WINDOWS
 
diff --git a/Engine/Plugins/NNE/NNERuntimeORT/Source/NNERuntimeORT/Private/NNERuntimeORTUtils.cpp b/Engine/Plugins/NNE/NNERuntimeORT/Source/NNERuntimeORT/Private/NNERuntimeORTUtils.cpp
--- a/Engine/Plugins/NNE/NNERuntimeORT/Source/NNERuntimeORT/Private/NNERuntimeORTUtils.cpp
+++ b/Engine/Plugins/NNE/NNERuntimeORT/Source/NNERuntimeORT/Private/NNERuntimeORTUtils.cpp
@@ -10,6 +10,8 @@
 #include "NNERuntimeORTEnv.h"
 
 #if PLATFORM_WINDOWS
+#include <dxcore_interface.h>
+#include <dxcore.h>
 #include "ID3D12DynamicRHI.h"
 #endif // PLATFORM_WINDOWS
 
@@ -32,6 +34,83 @@ static TAutoConsoleVariable<bool> CVarNNERuntimeORTEnableProfiling(
 
 namespace UE::NNERuntimeORT::Private
 {
+// Check for DirectX 12-compatible hardware.
+// Manually load DXCore.dll and d3d12.dll (avoid dll dependency) to enumerate adapters and
+// try to create a d3d12 device using the default adapter
+bool IsD3D12Available()
+{
+#if PLATFORM_WINDOWS
+	using Microsoft::WRL::ComPtr;
+	using DXCoreCreateAdapterFactoryFn = HRESULT __stdcall(REFIID, void**);
+
+	const int32 DeviceIndex = 0;
+
+	void* DxCoreModule = FPlatformProcess::GetDllHandle(TEXT("DXCore.dll"));
+	if (!DxCoreModule)
+	{
+		return false;
+	}
+
+	DXCoreCreateAdapterFactoryFn* DxCoreCreateAdapterFactory = reinterpret_cast<DXCoreCreateAdapterFactoryFn*>(FPlatformProcess::GetDllExport(DxCoreModule, TEXT("DXCoreCreateAdapterFactory")));
+	if (!DxCoreCreateAdapterFactory)
+	{
+		return false;
+	}
+
+	ComPtr<IDXCoreAdapterFactory> Factory;
+	DxCoreCreateAdapterFactory(IID_PPV_ARGS(&Factory));
+	if (!Factory)
+	{
+		return false;
+	}
+
+	const GUID DxGUIDs[] = { DXCORE_ADAPTER_ATTRIBUTE_D3D12_CORE_COMPUTE };
+	
+	ComPtr<IDXCoreAdapterList> AdapterList;
+	Factory->CreateAdapterList(ARRAYSIZE(DxGUIDs), DxGUIDs, IID_PPV_ARGS(&AdapterList));
+	if (!AdapterList || AdapterList->GetAdapterCount() < 1)
+	{
+		return false;
+	}
+
+	if (AdapterList->GetAdapterCount() <= DeviceIndex)
+	{
+		UE_LOG(LogNNE, Error, TEXT("Invalid device index %d. Number of available devices is %d."), DeviceIndex, AdapterList->GetAdapterCount());
+		return false;
+	}
+
+	ComPtr<IDXCoreAdapter> Adapter;
+	AdapterList->GetAdapter(static_cast<uint32_t>(DeviceIndex), IID_PPV_ARGS(&Adapter));
+	if (!Adapter)
+	{
+		return false;
+	}
+
+	void* D3D12Module = FPlatformProcess::GetDllHandle(TEXT("d3d12.dll"));
+	if (!D3D12Module)
+	{
+		return false;
+	}
+
+	decltype(&D3D12CreateDevice) D3D12CreateDeviceFun = reinterpret_cast<decltype(&D3D12CreateDevice)>(FPlatformProcess::GetDllExport(D3D12Module, TEXT("D3D12CreateDevice")));
+	if (!D3D12CreateDeviceFun)
+	{
+		return false;
+	}
+
+	ComPtr<ID3D12Device1> Device;
+	D3D12CreateDeviceFun(Adapter.Get(), D3D_FEATURE_LEVEL_1_0_CORE, DML_PPV_ARGS(&Device));
+	if (!Device)
+	{
+		return false;
+	}
+
+	return true;
+#else
+	return false;
+#endif // PLATFORM_WINDOWS
+}
+
 // For more details about ORT graph optimization checkout
 // https://onnxruntime.ai/docs/performance/model-optimizations/graph-optimizations.html
 
@@ -153,9 +232,9 @@ TUniquePtr<Ort::SessionOptions> CreateSessionOptionsDefault(const TSharedRef<FEn
 	return SessionOptions;
 }
 
-#if PLATFORM_WINDOWS
-TUniquePtr<Ort::SessionOptions> CreateSessionOptionsForDirectML(const TSharedRef<FEnvironment> &Environment)
+TUniquePtr<Ort::SessionOptions> CreateSessionOptionsForDirectML(const TSharedRef<FEnvironment> &Environment, bool bRHID3D12Required)
 {
+#if PLATFORM_WINDOWS
 	TUniquePtr<Ort::SessionOptions> SessionOptions = CreateSessionOptionsDefault(Environment);
 	if (!SessionOptions.IsValid())
 	{
@@ -166,8 +245,27 @@ TUniquePtr<Ort::SessionOptions> CreateSessionOptionsForDirectML(const TSharedRef
 	SessionOptions->SetExecutionMode(ExecutionMode::ORT_SEQUENTIAL);
 	SessionOptions->DisableMemPattern();
 
-	// In order to use DirectML we need D3D12
-	ID3D12DynamicRHI* RHI = nullptr;
+	if (!bRHID3D12Required && !IsRHID3D12())
+	{
+		const int32 DeviceIndex = 0;
+
+		const OrtDmlApi* DmlApi = nullptr;
+		Ort::ThrowOnError(Ort::GetApi().GetExecutionProviderApi("DML", ORT_API_VERSION, reinterpret_cast<const void**>(&DmlApi)));
+		if (!DmlApi)
+		{
+			UE_LOG(LogNNE, Error, TEXT("Ort DirectML Api not available!"));
+			return {};
+		}
+
+		OrtStatusPtr Status = DmlApi->SessionOptionsAppendExecutionProvider_DML(*SessionOptions.Get(), DeviceIndex);
+		if (Status)
+		{
+			UE_LOG(LogNNE, Error, TEXT("Failed to add DirectML execution provider to OnnxRuntime session options: %s"), ANSI_TO_TCHAR(Ort::GetApi().GetErrorMessage(Status)));
+			return {};
+		}
+
+		return SessionOptions;
+	}
 
 	if (!GDynamicRHI)
 	{
@@ -175,6 +273,8 @@ TUniquePtr<Ort::SessionOptions> CreateSessionOptionsForDirectML(const TSharedRef
 		return {};
 	}
 
+	// In order to use DirectML we need D3D12
+	ID3D12DynamicRHI* RHI = nullptr;
 	if (IsRHID3D12() )
 	{
 		RHI = GetID3D12DynamicRHI();
@@ -241,8 +341,10 @@ TUniquePtr<Ort::SessionOptions> CreateSessionOptionsForDirectML(const TSharedRef
 	}
 
 	return SessionOptions;
-}
+#else
+	return {};
 #endif // PLATFORM_WINDOWS
+}
 
 bool OptimizeModel(const TSharedRef<FEnvironment> &Environment, Ort::SessionOptions &SessionOptions, ENNEInferenceFormat TargetFormat, FNNEModelRaw& Model)
 {
diff --git a/Engine/Plugins/NNE/NNERuntimeORT/Source/NNERuntimeORT/Private/NNERuntimeORTUtils.h b/Engine/Plugins/NNE/NNERuntimeORT/Source/NNERuntimeORT/Private/NNERuntimeORTUtils.h
--- a/Engine/Plugins/NNE/NNERuntimeORT/Source/NNERuntimeORT/Private/NNERuntimeORTUtils.h
+++ b/Engine/Plugins/NNE/NNERuntimeORT/Source/NNERuntimeORT/Private/NNERuntimeORTUtils.h
@@ -13,6 +13,8 @@ namespace UE::NNERuntimeORT::Private
 
 	class FEnvironment;
 
+	bool IsD3D12Available();
+
 namespace OrtHelper
 {
 	TArray<uint32> GetShape(const Ort::Value& OrtTensor);
@@ -24,9 +26,7 @@ namespace OrtHelper
 
 	TUniquePtr<Ort::SessionOptions> CreateSessionOptionsDefault(const TSharedRef<FEnvironment> &Environment);
 
-#if PLATFORM_WINDOWS
-	TUniquePtr<Ort::SessionOptions> CreateSessionOptionsForDirectML(const TSharedRef<FEnvironment> &Environment);
-#endif // PLATFORM_WINDOWS
+	TUniquePtr<Ort::SessionOptions> CreateSessionOptionsForDirectML(const TSharedRef<FEnvironment> &Environment, bool bRHID3D12Required = true);
 
 	bool OptimizeModel(const TSharedRef<FEnvironment> &Environment, Ort::SessionOptions &SessionOptions, ENNEInferenceFormat TargetFormat, FNNEModelRaw& Model);
 
