[NNE] NNERuntimeORT minor fixes and cleanup.

#tests Editor, Win64 smoke tests.

#rnx

diff --git a/Engine/Plugins/NNE/NNERuntimeORT/Source/NNERuntimeORT/Private/NNERuntimeORTModel.cpp b/Engine/Plugins/NNE/NNERuntimeORT/Source/NNERuntimeORT/Private/NNERuntimeORTModel.cpp
--- a/Engine/Plugins/NNE/NNERuntimeORT/Source/NNERuntimeORT/Private/NNERuntimeORTModel.cpp
+++ b/Engine/Plugins/NNE/NNERuntimeORT/Source/NNERuntimeORT/Private/NNERuntimeORTModel.cpp
@@ -665,7 +665,24 @@ FModelInstanceORTDmlRDG::ESetInputTensorShapesStatus FModelInstanceORTDmlRDG::Se
 	TConstArrayView<uint8> ModelBuffer = TConstArrayView<uint8>(&(ModelData->GetView().GetData()[GuidSize + VersionSize]), ModelData->GetView().Num() - GuidSize - VersionSize);
 	check(!ModelBuffer.IsEmpty());
 
-	Session = MakeUnique<Ort::Session>(Environment->GetOrtEnv(), ModelBuffer.GetData(), ModelBuffer.Num(), *SessionOptions);
+#if WITH_EDITOR
+	try
+#endif // WITH_EDITOR
+	{
+		Session = MakeUnique<Ort::Session>(Environment->GetOrtEnv(), ModelBuffer.GetData(), ModelBuffer.Num(), *SessionOptions);
+	}
+#if WITH_EDITOR
+	catch (const Ort::Exception& Exception)
+	{
+		UE_LOG(LogNNE, Error, TEXT("%s"), UTF8_TO_TCHAR(Exception.what()));
+		return ESetInputTensorShapesStatus::Fail;
+	}
+	catch (...)
+	{
+		UE_LOG(LogNNE, Error, TEXT("Unknown exception!"));
+		return ESetInputTensorShapesStatus::Fail;
+	}
+#endif // WITH_EDITOR
 
 	// Need to configure output tensors with new session (to apply free dimension overrides)
 	if (!ConfigureTensors(*Session, false))
diff --git a/Engine/Plugins/NNE/NNERuntimeORT/Source/NNERuntimeORT/Private/NNERuntimeORTSettings.h b/Engine/Plugins/NNE/NNERuntimeORT/Source/NNERuntimeORT/Private/NNERuntimeORTSettings.h
--- a/Engine/Plugins/NNE/NNERuntimeORT/Source/NNERuntimeORT/Private/NNERuntimeORTSettings.h
+++ b/Engine/Plugins/NNE/NNERuntimeORT/Source/NNERuntimeORT/Private/NNERuntimeORTSettings.h
@@ -14,7 +14,7 @@ enum EExecutionMode : uint8
 	PARALLEL
 };
 
-/** Rigid body error correction data */
+/** Threading options to configure glocal/local thread pools and exection mode */
 USTRUCT()
 struct FThreadingOptions
 {
diff --git a/Engine/Plugins/NNE/NNERuntimeORT/Source/NNERuntimeORT/Private/NNERuntimeORTUtils.cpp b/Engine/Plugins/NNE/NNERuntimeORT/Source/NNERuntimeORT/Private/NNERuntimeORTUtils.cpp
--- a/Engine/Plugins/NNE/NNERuntimeORT/Source/NNERuntimeORT/Private/NNERuntimeORTUtils.cpp
+++ b/Engine/Plugins/NNE/NNERuntimeORT/Source/NNERuntimeORT/Private/NNERuntimeORTUtils.cpp
@@ -217,7 +217,8 @@ TUniquePtr<Ort::SessionOptions> CreateSessionOptionsDefault(const TSharedRef<FEn
 	}
 
 	// Configure Profiling
-	if (CVarNNERuntimeORTEnableProfiling.GetValueOnGameThread())
+	// Note: can be called on game or render thread
+	if (CVarNNERuntimeORTEnableProfiling.GetValueOnAnyThread())
 	{
 		FString ProfilingFilePrefix("NNERuntimeORTProfile_");
 		ProfilingFilePrefix += FString::FromInt(ORTProfilingSessionNumber);
