[NNE] RuntimeORT addendum to head-less CL. Use DXGI instead of DXCore to not break with the minimum Windows SDK version.

#tests Editor, Win64 run tests in head-less mode

#rnx

diff --git a/Engine/Plugins/NNE/NNERuntimeORT/Source/NNERuntimeORT/Private/NNERuntimeORTUtils.cpp b/Engine/Plugins/NNE/NNERuntimeORT/Source/NNERuntimeORT/Private/NNERuntimeORTUtils.cpp
--- a/Engine/Plugins/NNE/NNERuntimeORT/Source/NNERuntimeORT/Private/NNERuntimeORTUtils.cpp
+++ b/Engine/Plugins/NNE/NNERuntimeORT/Source/NNERuntimeORT/Private/NNERuntimeORTUtils.cpp
@@ -10,8 +10,7 @@
 #include "NNERuntimeORTEnv.h"
 
 #if PLATFORM_WINDOWS
-#include <dxcore_interface.h>
-#include <dxcore.h>
+#include <dxgi1_4.h>
 #include "ID3D12DynamicRHI.h"
 #endif // PLATFORM_WINDOWS
 
@@ -35,71 +34,32 @@ static TAutoConsoleVariable<bool> CVarNNERuntimeORTEnableProfiling(
 namespace UE::NNERuntimeORT::Private
 {
 // Check for DirectX 12-compatible hardware.
-// Manually load DXCore.dll and d3d12.dll (avoid dll dependency) to enumerate adapters and
-// try to create a d3d12 device using the default adapter
+// Use DXGI to enumerate adapters and try to create a d3d12 device using the default adapter (will create dependency to dxgi.dll!)
+// DXGI 1.6 should be available since Windows 10, version 1809, which is newer than the minimum SDK version
+// specified in Engine\Config\Windows\Windows_SDK.json at the moment.
 bool IsD3D12Available()
 {
 #if PLATFORM_WINDOWS
 	using Microsoft::WRL::ComPtr;
-	using DXCoreCreateAdapterFactoryFn = HRESULT __stdcall(REFIID, void**);
 
 	const int32 DeviceIndex = 0;
 
-	void* DxCoreModule = FPlatformProcess::GetDllHandle(TEXT("DXCore.dll"));
-	if (!DxCoreModule)
-	{
-		return false;
-	}
-
-	DXCoreCreateAdapterFactoryFn* DxCoreCreateAdapterFactory = reinterpret_cast<DXCoreCreateAdapterFactoryFn*>(FPlatformProcess::GetDllExport(DxCoreModule, TEXT("DXCoreCreateAdapterFactory")));
-	if (!DxCoreCreateAdapterFactory)
-	{
-		return false;
-	}
-
-	ComPtr<IDXCoreAdapterFactory> Factory;
-	DxCoreCreateAdapterFactory(IID_PPV_ARGS(&Factory));
+	ComPtr<IDXGIFactory4> Factory;
+	CreateDXGIFactory2(0, IID_PPV_ARGS(&Factory));
 	if (!Factory)
 	{
 		return false;
 	}
 
-	const GUID DxGUIDs[] = { DXCORE_ADAPTER_ATTRIBUTE_D3D12_CORE_COMPUTE };
-	
-	ComPtr<IDXCoreAdapterList> AdapterList;
-	Factory->CreateAdapterList(ARRAYSIZE(DxGUIDs), DxGUIDs, IID_PPV_ARGS(&AdapterList));
-	if (!AdapterList || AdapterList->GetAdapterCount() < 1)
-	{
-		return false;
-	}
-
-	if (AdapterList->GetAdapterCount() <= DeviceIndex)
-	{
-		UE_LOG(LogNNE, Error, TEXT("Invalid device index %d. Number of available devices is %d."), DeviceIndex, AdapterList->GetAdapterCount());
-		return false;
-	}
-
-	ComPtr<IDXCoreAdapter> Adapter;
-	AdapterList->GetAdapter(static_cast<uint32_t>(DeviceIndex), IID_PPV_ARGS(&Adapter));
+	ComPtr<IDXGIAdapter1> Adapter;
+	Factory->EnumAdapters1(DeviceIndex, &Adapter);
 	if (!Adapter)
 	{
 		return false;
 	}
 
-	void* D3D12Module = FPlatformProcess::GetDllHandle(TEXT("d3d12.dll"));
-	if (!D3D12Module)
-	{
-		return false;
-	}
-
-	decltype(&D3D12CreateDevice) D3D12CreateDeviceFun = reinterpret_cast<decltype(&D3D12CreateDevice)>(FPlatformProcess::GetDllExport(D3D12Module, TEXT("D3D12CreateDevice")));
-	if (!D3D12CreateDeviceFun)
-	{
-		return false;
-	}
-
-	ComPtr<ID3D12Device1> Device;
-	D3D12CreateDeviceFun(Adapter.Get(), D3D_FEATURE_LEVEL_1_0_CORE, DML_PPV_ARGS(&Device));
+	ComPtr<ID3D12Device> Device;
+	D3D12CreateDevice(Adapter.Get(), D3D_FEATURE_LEVEL_11_0, IID_PPV_ARGS(&Device));
 	if (!Device)
 	{
 		return false;
