[NNE] NNERuntimeORT write settings to default config.

#tests Editor, Win64 verify settings are persistent after restart

#rnx

diff --git a/Engine/Plugins/NNE/NNERuntimeORT/Source/NNERuntimeORT/Private/NNERuntimeORTSettings.cpp b/Engine/Plugins/NNE/NNERuntimeORT/Source/NNERuntimeORT/Private/NNERuntimeORTSettings.cpp
new file mode 100644
--- /dev/null
+++ b/Engine/Plugins/NNE/NNERuntimeORT/Source/NNERuntimeORT/Private/NNERuntimeORTSettings.cpp
@@ -0,0 +1,21 @@
+// Copyright Epic Games, Inc. All Rights Reserved.
+
+#include "NNERuntimeORTSettings.h"
+
+UNNERuntimeORTSettings::UNNERuntimeORTSettings(const FObjectInitializer& ObjectInitlaizer)
+	: Super(ObjectInitlaizer)
+{
+
+}
+
+FName UNNERuntimeORTSettings::GetCategoryName() const
+{
+	return TEXT("Plugins");
+}
+
+#if WITH_EDITOR
+FText UNNERuntimeORTSettings::GetSectionText() const
+{
+	return NSLOCTEXT("NNERuntimeORTPlugin", "NNERuntimeORTSettingsSection", "NNERuntimeORT");
+}
+#endif
\ No newline at end of file
diff --git a/Engine/Plugins/NNE/NNERuntimeORT/Source/NNERuntimeORT/Private/NNERuntimeORTSettings.h b/Engine/Plugins/NNE/NNERuntimeORT/Source/NNERuntimeORT/Private/NNERuntimeORTSettings.h
--- a/Engine/Plugins/NNE/NNERuntimeORT/Source/NNERuntimeORT/Private/NNERuntimeORTSettings.h
+++ b/Engine/Plugins/NNE/NNERuntimeORT/Source/NNERuntimeORT/Private/NNERuntimeORTSettings.h
@@ -46,11 +46,13 @@ struct FThreadingOptions
 };
 
 /** Settings used to configure NNERuntimeORT */
-UCLASS(Config = Engine, meta = (DisplayName = "NNERuntimeORT"))
-class NNERUNTIMEORT_API UNNERuntimeORTSettings : public UDeveloperSettings
+UCLASS(config = Engine, defaultconfig, meta = (DisplayName = "NNERuntimeORT"))
+class UNNERuntimeORTSettings : public UDeveloperSettings
 {
 	GENERATED_BODY()
 
+	UNNERuntimeORTSettings(const FObjectInitializer& ObjectInitlaizer);
+
 public:
 
 	/** Threading options in Editor targets */
@@ -61,4 +63,11 @@ public:
 	UPROPERTY(Config, EditAnywhere, Category = "ONNX Runtime", meta = (DisplayName = "Game threading options", ToolTip = "Threading options in Non-Editor (Game, Program, ...) targets", ConfigRestartRequired = true))
 	FThreadingOptions GameThreadingOptions{false, 1, 1, EExecutionMode::SEQUENTIAL};
 
+	// Begin UDeveloperSettings Interface
+	NNERUNTIMEORT_API virtual FName GetCategoryName() const override;
+#if WITH_EDITOR
+	NNERUNTIMEORT_API virtual FText GetSectionText() const override;
+	// END UDeveloperSettings Interface
+#endif
+
 };
