[NNE] Enable NNERuntimeORTDml on all Editor platforms.

#tests Editor, Win64, Mac compilation and cook tests

#rnx

diff --git a/Engine/Plugins/NNE/NNERuntimeORT/Source/NNERuntimeORT/Private/NNERuntimeORTModule.cpp b/Engine/Plugins/NNE/NNERuntimeORT/Source/NNERuntimeORT/Private/NNERuntimeORTModule.cpp
--- a/Engine/Plugins/NNE/NNERuntimeORT/Source/NNERuntimeORT/Private/NNERuntimeORTModule.cpp
+++ b/Engine/Plugins/NNE/NNERuntimeORT/Source/NNERuntimeORT/Private/NNERuntimeORTModule.cpp
@@ -74,11 +74,13 @@ void FNNERuntimeORTModule::StartupModule()
 
 	void* OrtDllHandle = DllHandles.Last();
 
+	bool bDirectMLDllLoaded = false;
+
 #if PLATFORM_WINDOWS
 	const FString ModuleDir = FPlatformProcess::GetModulesDirectory();
 	const FString DirectMLSharedLibPath = FPaths::Combine(ModuleDir, TEXT(PREPROCESSOR_TO_STRING(DIRECTML_PATH)), TEXT("DirectML.dll"));
 
-	const bool bDirectMLDllLoaded = DllHelper::GetDllHandle(DirectMLSharedLibPath, DllHandles);
+	bDirectMLDllLoaded = DllHelper::GetDllHandle(DirectMLSharedLibPath, DllHandles);
 	if (!bDirectMLDllLoaded)
 	{
 		UE_LOG(LogNNE, Error, TEXT("Failed to load DirectML shared library. ORT Dml Runtime won't be available."));
@@ -98,7 +100,6 @@ void FNNERuntimeORTModule::StartupModule()
 
 	EnvironmentHelper::CreateOrtEnvFromSettings(GetDefault<UNNERuntimeORTSettings>(), *Environment);
 
-#if PLATFORM_WINDOWS
 	// NNE runtime ORT Dml startup
 	NNERuntimeORTDml = NewObject<UNNERuntimeORTDml>();
 	if (NNERuntimeORTDml.IsValid())
@@ -109,7 +110,6 @@ void FNNERuntimeORTModule::StartupModule()
 		NNERuntimeORTDml->AddToRoot();
 		UE::NNE::RegisterRuntime(RuntimeDmlInterface);
 	}
-#endif // PLATFORM_WINDOWS
 
 	// NNE runtime ORT Cpu startup
 	NNERuntimeORTCpu = NewObject<UNNERuntimeORTCpu>();
