NNE - Onnx External data support (import only on windows)

Problem: Onnx model with data in external files are not supported.

Solution:
* Wrap onnx as a DLL and integrate it to NNEEditor module (on windows).
* Specialize the NNEEditor importer to be an onnx only importer.
* Update SharedModelData for the ORT based runtimes to support external file data
* For ORT runtime recreate the additional files hierarchy on disk when a creating a session from a model with external data (this is because we don't yet have a nice way to use only memory in that case, ORT v18 will bring that).
* Refactor NNEQA to import onnx model using the NNEEditor code path.

Bonus:
* General cleanup of the editor vs runtime dependencies.

Test:
* Added ONNX model with external data to NNE tests (one data file, multiple data files, one data file within subfolders)
* Win64 Editor Debug: all tests (including reimport) + manual import and reimport of models with and without external data + style transfer (hlsl)
* Win64 Standalone Debug packaged: all tests + style transfer (hlsl)
* Linux Editor Development: all tests (including model with external file imported via windows) + verification import of models with external data fail as expected and import of model without additional data works also as expected.
* Linux Standalone Development packaged : all tests (including model with external imported via windows).
* Console Debug Standalone: smoke tests + style transfer (hlsl)

#rnx

#virtualized

diff --git a/Engine/Plugins/Experimental/NNERuntimeRDG/Source/NNERuntimeRDG/Private/NNERuntimeRDGHlsl.cpp b/Engine/Plugins/Experimental/NNERuntimeRDG/Source/NNERuntimeRDG/Private/NNERuntimeRDGHlsl.cpp
--- a/Engine/Plugins/Experimental/NNERuntimeRDG/Source/NNERuntimeRDG/Private/NNERuntimeRDGHlsl.cpp
+++ b/Engine/Plugins/Experimental/NNERuntimeRDG/Source/NNERuntimeRDG/Private/NNERuntimeRDGHlsl.cpp
@@ -183,6 +183,12 @@ UNNERuntimeRDGHlslImpl::ECanCreateModelRDGStatus UNNERuntimeRDGHlslImpl::CanCrea
 
 TSharedPtr<UE::NNE::FSharedModelData> UNNERuntimeRDGHlslImpl::CreateModelData(const FString& FileType, TConstArrayView<uint8> FileData, const TMap<FString, TConstArrayView<uint8>>& AdditionalFileData, const FGuid& FileId, const ITargetPlatform* TargetPlatform)
 {
+	if (!AdditionalFileData.IsEmpty())
+	{
+		UE_LOG(LogNNE, Warning, TEXT("UNNERuntimeRDGHlsl cannot create the model data with id %s (Filetype: %s) external data not supported at the moment, please convert the model to internal storage. See https://onnx.ai/onnx/repo-docs/ExternalData.html"), *FileId.ToString(EGuidFormats::Digits).ToLower(), *FileType);
+		return {};
+	}
+
 	if (CanCreateModelData(FileType, FileData, AdditionalFileData, FileId, TargetPlatform) != ECanCreateModelDataStatus::Ok)
 	{
 		UE_LOG(LogNNE, Warning, TEXT("UNNERuntimeRDGHlsl cannot create the model data with id %s (Filetype: %s)"), *FileId.ToString(EGuidFormats::Digits).ToLower(), *FileType);
diff --git a/Engine/Plugins/NNE/NNERuntimeORT/Source/NNERuntimeORT/Private/NNERuntimeORT.cpp b/Engine/Plugins/NNE/NNERuntimeORT/Source/NNERuntimeORT/Private/NNERuntimeORT.cpp
--- a/Engine/Plugins/NNE/NNERuntimeORT/Source/NNERuntimeORT/Private/NNERuntimeORT.cpp
+++ b/Engine/Plugins/NNE/NNERuntimeORT/Source/NNERuntimeORT/Private/NNERuntimeORT.cpp
@@ -10,6 +10,7 @@
 #include "NNEModelData.h"
 #include "NNEModelOptimizerInterface.h"
 #include "NNERuntimeORTModel.h"
+#include "NNERuntimeORTModelFormat.h"
 #include "NNERuntimeORTUtils.h"
 
 #if PLATFORM_WINDOWS
@@ -19,10 +20,61 @@
 #include UE_INLINE_GENERATED_CPP_BY_NAME(NNERuntimeORT)
 
 FGuid UNNERuntimeORTCpu::GUID = FGuid((int32)'O', (int32)'C', (int32)'P', (int32)'U');
-int32 UNNERuntimeORTCpu::Version = 0x00000002;
+int32 UNNERuntimeORTCpu::Version = 0x00000003;
 
 FGuid UNNERuntimeORTDml::GUID = FGuid((int32)'O', (int32)'D', (int32)'M', (int32)'L');
-int32 UNNERuntimeORTDml::Version = 0x00000002;
+int32 UNNERuntimeORTDml::Version = 0x00000003;
+
+namespace UE::NNERuntimeORT::Private::Details
+{
+	//Should be kept in sync with OnnxFileLoaderHelper::InitUNNEModelDataFromFile()
+	static FString OnnxExternalDataDescriptorKey(TEXT("OnnxExternalDataDescriptor"));
+	static FString OnnxExternalDataBytesKey(TEXT("OnnxExternalDataBytes"));
+
+	FOnnxDataDescriptor MakeOnnxDataDescriptor(TConstArrayView<uint8> FileData, const TMap<FString, TConstArrayView<uint8>>& AdditionalFileData)
+	{
+		FOnnxDataDescriptor OnnxDataDescriptor = {};
+		OnnxDataDescriptor.OnnxModelDataSize = FileData.Num();
+
+		if (AdditionalFileData.Contains(OnnxExternalDataDescriptorKey))
+		{
+			TConstArrayView<uint8> OnnxExternalDataDescriptorBuffer = AdditionalFileData[OnnxExternalDataDescriptorKey];
+			FMemoryReaderView OnnxExternalDataDescriptorReader(OnnxExternalDataDescriptorBuffer, /*bIsPersistent = */true);
+			TMap<FString, int64> ExternalDataSizes;
+
+			OnnxExternalDataDescriptorReader << ExternalDataSizes;
+
+			int64 CurrentBucketOffset = OnnxDataDescriptor.OnnxModelDataSize;
+			for (const auto& Element : ExternalDataSizes)
+			{
+				const FString DataFilePath = Element.Key;
+				OnnxDataDescriptor.AdditionalDataDescriptors.Emplace(
+					DataFilePath,
+					CurrentBucketOffset,
+					Element.Value);
+				CurrentBucketOffset += Element.Value;
+			}
+		}
+		return OnnxDataDescriptor;
+	}
+
+	void WriteOnnxModelData(FMemoryWriter Writer, TConstArrayView<uint8> FileData, const TMap<FString, TConstArrayView<uint8>>& AdditionalFileData)
+	{
+		FOnnxDataDescriptor Descriptor = MakeOnnxDataDescriptor(FileData, AdditionalFileData);
+		check(FileData.Num() == Descriptor.OnnxModelDataSize);
+		Writer << Descriptor;
+
+		Writer.Serialize(const_cast<uint8*>(FileData.GetData()), FileData.Num());
+
+		if (!Descriptor.AdditionalDataDescriptors.IsEmpty())
+		{
+			
+
+			check(AdditionalFileData.Contains(OnnxExternalDataBytesKey));
+			Writer.Serialize(const_cast<uint8*>(AdditionalFileData[OnnxExternalDataBytesKey].GetData()), AdditionalFileData[OnnxExternalDataBytesKey].Num());
+		}
+	}
+} // namespace UE::NNERuntimeORT::Private::Details
 
 UNNERuntimeORTCpu::ECanCreateModelDataStatus UNNERuntimeORTCpu::CanCreateModelData(const FString& FileType, TConstArrayView<uint8> FileData, const TMap<FString, TConstArrayView<uint8>>& AdditionalFileData, const FGuid& FileId, const ITargetPlatform* TargetPlatform) const
 {
@@ -39,24 +91,30 @@ TSharedPtr<UE::NNE::FSharedModelData> UNNERuntimeORTCpu::CreateModelData(const F
 		return {};
 	}
 
-	FNNEModelRaw InputModel{TArray<uint8>{FileData}, ENNEInferenceFormat::ONNX};
-	if (GraphOptimizationLevel OptimizationLevel = GetGraphOptimizationLevelForCPU(false, IsRunningCookCommandlet()); OptimizationLevel > GraphOptimizationLevel::ORT_DISABLE_ALL)
-	{
-		TUniquePtr<Ort::SessionOptions> SessionOptions = CreateSessionOptionsDefault(Environment.ToSharedRef());
-		SessionOptions->SetGraphOptimizationLevel(OptimizationLevel);
-		SessionOptions->EnableCpuMemArena();
+	FNNEModelRaw InputModel{ TArray<uint8>{FileData}, ENNEInferenceFormat::ONNX };
 
-		if (!OptimizeModel(Environment.ToSharedRef(), *SessionOptions, ENNEInferenceFormat::ONNX, InputModel))
+	//For now only optimize model if there is no external data (as additional data are serialized from the unoptimized model below)
+	if (AdditionalFileData.IsEmpty())
+	{
+		if (GraphOptimizationLevel OptimizationLevel = GetGraphOptimizationLevelForCPU(false, IsRunningCookCommandlet()); OptimizationLevel > GraphOptimizationLevel::ORT_DISABLE_ALL)
 		{
-			return {};
+			TUniquePtr<Ort::SessionOptions> SessionOptions = CreateSessionOptionsDefault(Environment.ToSharedRef());
+			SessionOptions->SetGraphOptimizationLevel(OptimizationLevel);
+			SessionOptions->EnableCpuMemArena();
+
+			if (!OptimizeModel(Environment.ToSharedRef(), *SessionOptions, ENNEInferenceFormat::ONNX, InputModel))
+			{
+				return {};
+			}
 		}
 	}
 
 	TArray<uint8> Result;
-	FMemoryWriter Writer(Result);
+	FMemoryWriter Writer(Result, /*bIsPersitent =*/ true);
 	Writer << UNNERuntimeORTCpu::GUID;
 	Writer << UNNERuntimeORTCpu::Version;
-	Writer.Serialize(InputModel.Data.GetData(), InputModel.Data.Num());
+
+	Details::WriteOnnxModelData(Writer, InputModel.Data, AdditionalFileData);
 
 	return MakeShared<UE::NNE::FSharedModelData>(MakeSharedBufferFromArray(MoveTemp(Result)), 0);
 }
@@ -161,25 +219,31 @@ TSharedPtr<UE::NNE::FSharedModelData> UNNERuntimeORTDml::CreateModelData(const F
 		return {};
 	}
 
-	FNNEModelRaw InputModel{TArray<uint8>{FileData}, ENNEInferenceFormat::ONNX};
-	if (GraphOptimizationLevel OptimizationLevel = GetGraphOptimizationLevelForDML(false, IsRunningCookCommandlet()); OptimizationLevel > GraphOptimizationLevel::ORT_DISABLE_ALL)
+	FNNEModelRaw InputModel{ TArray<uint8>{FileData}, ENNEInferenceFormat::ONNX };
+
+	//For now only optimize model if there is no external data (as additional data are serialized from the unoptimized model below)
+	if (AdditionalFileData.IsEmpty())
 	{
-		TUniquePtr<Ort::SessionOptions> SessionOptions = CreateSessionOptionsDefault(Environment.ToSharedRef());
-		SessionOptions->SetGraphOptimizationLevel(OptimizationLevel);
-		SessionOptions->SetExecutionMode(ExecutionMode::ORT_SEQUENTIAL);
-		SessionOptions->DisableMemPattern();
-		
-		if (!OptimizeModel(Environment.ToSharedRef(), *SessionOptions, ENNEInferenceFormat::ONNX, InputModel))
+		if (GraphOptimizationLevel OptimizationLevel = GetGraphOptimizationLevelForDML(false, IsRunningCookCommandlet()); OptimizationLevel > GraphOptimizationLevel::ORT_DISABLE_ALL)
 		{
-			return {};
+			TUniquePtr<Ort::SessionOptions> SessionOptions = CreateSessionOptionsDefault(Environment.ToSharedRef());
+			SessionOptions->SetGraphOptimizationLevel(OptimizationLevel);
+			SessionOptions->SetExecutionMode(ExecutionMode::ORT_SEQUENTIAL);
+			SessionOptions->DisableMemPattern();
+
+			if (!OptimizeModel(Environment.ToSharedRef(), *SessionOptions, ENNEInferenceFormat::ONNX, InputModel))
+			{
+				return {};
+			}
 		}
 	}
 
 	TArray<uint8> Result;
-	FMemoryWriter Writer(Result);
+	FMemoryWriter Writer(Result, /*bIsPersitent =*/ true);
 	Writer << UNNERuntimeORTDml::GUID;
 	Writer << UNNERuntimeORTDml::Version;
-	Writer.Serialize(InputModel.Data.GetData(), InputModel.Data.Num());
+
+	Details::WriteOnnxModelData(Writer, InputModel.Data, AdditionalFileData);
 
 	return MakeShared<UE::NNE::FSharedModelData>(MakeSharedBufferFromArray(MoveTemp(Result)), 0);
 }
diff --git a/Engine/Plugins/NNE/NNERuntimeORT/Source/NNERuntimeORT/Private/NNERuntimeORTModel.cpp b/Engine/Plugins/NNE/NNERuntimeORT/Source/NNERuntimeORT/Private/NNERuntimeORTModel.cpp
--- a/Engine/Plugins/NNE/NNERuntimeORT/Source/NNERuntimeORT/Private/NNERuntimeORTModel.cpp
+++ b/Engine/Plugins/NNE/NNERuntimeORT/Source/NNERuntimeORT/Private/NNERuntimeORTModel.cpp
@@ -3,7 +3,11 @@
 #include "NNERuntimeORTModel.h"
 
 #include "HAL/PlatformMisc.h"
+#include "HAL/PlatformFileManager.h"
+#include "Misc/FileHelper.h"
+#include "Misc/Paths.h"
 #include "NNERuntimeORT.h"
+#include "NNERuntimeORTModelFormat.h"
 #include "NNERuntimeORTSettings.h"
 #include "NNERuntimeORTUtils.h"
 #include "RenderGraphUtils.h"
@@ -45,6 +49,92 @@ namespace Detail
 		return Result;
 	}
 
+	FString CreateTempDirPath(const FString& BasePath)
+	{
+		FString UniqueDirName;
+		do
+		{
+			UniqueDirName = FPaths::Combine(BasePath, *FString::Printf(TEXT("ORTModel_%s"), *FGuid::NewGuid().ToString()));
+		} while (IFileManager::Get().DirectoryExists(*UniqueDirName));
+
+		return UniqueDirName;
+	}
+
+	bool CreateSession(
+		TConstArrayView<uint8> ModelData,
+		const Ort::SessionOptions& SessionOptions,
+		const FEnvironment& Environment,
+		TUniquePtr<Ort::Session>& Session, FString& TempDirForModelWithExternalData)
+	{
+		FMemoryReaderView Reader(ModelData, /*bIsPersitent =*/ true);
+		FGuid GUID;
+		int32 Version;
+		Reader << GUID;
+		Reader << Version;
+
+		FOnnxDataDescriptor Descriptor;
+		Reader << Descriptor;
+
+		int64 BaseDataOffset = Reader.Tell();
+		TConstArrayView<uint8> ModelBuffer = TConstArrayView<uint8>(&(ModelData.GetData()[BaseDataOffset]), Descriptor.OnnxModelDataSize);
+
+		if (ModelBuffer.Num() == 0)
+		{
+			UE_LOG(LogNNE, Error, TEXT("NNERuntimeORT::Private::Details::CreateSession(): Input model data is empty."));
+			return false;
+		}
+
+		// Starting with ORT v18 we will get AddExternalInitializersFromFilesInMemory() via onnxruntime_c_api.h
+		// however for now we use temp files when working with model with external data.
+		if (Descriptor.AdditionalDataDescriptors.Num() > 0)
+		{
+			FString Filepath;
+			if (TempDirForModelWithExternalData.IsEmpty())
+			{
+				FString ProjIntermediateDir = FPaths::ConvertRelativePathToFull(FPaths::ProjectIntermediateDir());
+				TempDirForModelWithExternalData = Detail::CreateTempDirPath(ProjIntermediateDir);
+				Filepath = FPaths::Combine(TempDirForModelWithExternalData, TEXT("OnnxModel.onnx"));
+
+				// Note: SaveArrayToFile() will create the needed folders as needed both for the Onnx model and the additional data files.
+				if (!FFileHelper::SaveArrayToFile(ModelBuffer, *Filepath))
+				{
+					IFileManager::Get().DeleteDirectory(*TempDirForModelWithExternalData, false, true);
+					UE_LOG(LogNNE, Error, TEXT("NNERuntimeORT::Private::Details::CreateSession() could not write model to disk at %s."), *Filepath);
+					return false;
+				}
+
+				for (const FOnnxAdditionalDataDescriptor& AdditionalDataDescriptor : Descriptor.AdditionalDataDescriptors)
+				{
+					FString AdditionalDataFilename = FPaths::Combine(TempDirForModelWithExternalData, *AdditionalDataDescriptor.Path);
+					TConstArrayView<uint8> AdditionalDataBuffer = TConstArrayView<uint8>(&(ModelData.GetData()[BaseDataOffset + AdditionalDataDescriptor.Offset]), AdditionalDataDescriptor.Size);
+
+					if (!FFileHelper::SaveArrayToFile(AdditionalDataBuffer, *AdditionalDataFilename))
+					{
+						IFileManager::Get().DeleteDirectory(*TempDirForModelWithExternalData, false, true);
+						UE_LOG(LogNNE, Error, TEXT("NNERuntimeORT::Private::Details::CreateSession() could not write additional data to disk at %s."), *AdditionalDataFilename);
+						return false;
+					}
+				}
+			}
+			else
+			{
+				Filepath = FPaths::Combine(TempDirForModelWithExternalData, TEXT("OnnxModel.onnx"));
+			}
+
+#if PLATFORM_WINDOWS
+			Session = MakeUnique<Ort::Session>(Environment.GetOrtEnv(), *Filepath, SessionOptions);
+#else
+			Session = MakeUnique<Ort::Session>(Environment.GetOrtEnv(), TCHAR_TO_ANSI(*Filepath), SessionOptions);
+#endif
+		}
+		else
+		{
+			Session = MakeUnique<Ort::Session>(Environment.GetOrtEnv(), ModelBuffer.GetData(), ModelBuffer.Num(), SessionOptions);
+		}
+
+		return Session.IsValid();
+	}
+
 } // namespace Detail
 
 template <class ModelInterface, class TensorBinding> 
@@ -54,19 +144,22 @@ FModelInstanceORTBase<ModelInterface, TensorBinding>::FModelInstanceORTBase(cons
 
 }
 
-template <class ModelInterface, class TensorBinding> 
-bool FModelInstanceORTBase<ModelInterface, TensorBinding>::Init(TConstArrayView<uint8> ModelData)
+template <class ModelInterface, class TensorBinding>
+FModelInstanceORTBase<ModelInterface, TensorBinding>::~FModelInstanceORTBase()
 {
-	constexpr int32 GuidSize = sizeof(UNNERuntimeORTDml::GUID);
-	constexpr int32 VersionSize = sizeof(UNNERuntimeORTDml::Version);
-	TConstArrayView<uint8> ModelBuffer = TConstArrayView<uint8>(&(ModelData.GetData()[GuidSize + VersionSize]), ModelData.Num() - GuidSize - VersionSize);
-
-	if (ModelBuffer.Num() == 0)
+	Session.Reset();
+	if (!TempDirForModelWithExternalData.IsEmpty())
 	{
-		UE_LOG(LogNNE, Error, TEXT("FModelInstanceORTBase::Init(): Input model data is empty."));
-		return false;
+		if (!IFileManager::Get().DeleteDirectory(*TempDirForModelWithExternalData, false, true))
+		{
+			UE_LOG(LogNNE, Warning, TEXT("FModelInstanceORTBase could not delete temp directy %s on model instance destruction."), *TempDirForModelWithExternalData);
+		}
 	}
+}
 
+template <class ModelInterface, class TensorBinding> 
+bool FModelInstanceORTBase<ModelInterface, TensorBinding>::Init(TConstArrayView<uint8> ModelData)
+{
 #if WITH_EDITOR
 	try
 #endif // WITH_EDITOR
@@ -77,7 +170,11 @@ bool FModelInstanceORTBase<ModelInterface, TensorBinding>::Init(TConstArrayView<
 			return false;
 		}
 
-		Session = MakeUnique<Ort::Session>(Environment->GetOrtEnv(), ModelBuffer.GetData(), ModelBuffer.Num(), *SessionOptions);
+		if (!Detail::CreateSession(ModelData, *SessionOptions, *Environment, Session, TempDirForModelWithExternalData))
+		{
+			UE_LOG(LogNNE, Error, TEXT("FModelInstanceORTBase::Init(): Session creation failed."));
+			return false;
+		}
 
 		if (!ConfigureTensors(true))
 		{
@@ -440,18 +537,20 @@ FModelInstanceORTDmlRDG::FModelInstanceORTDmlRDG(TSharedRef<UE::NNE::FSharedMode
 	:  ModelData(InModelData), RuntimeConf(InRuntimeConf), Environment(InEnvironment)
 {}
 
-bool FModelInstanceORTDmlRDG::Init()
+FModelInstanceORTDmlRDG::~FModelInstanceORTDmlRDG()
 {
-	constexpr int32 GuidSize = sizeof(UNNERuntimeORTDml::GUID);
-	constexpr int32 VersionSize = sizeof(UNNERuntimeORTDml::Version);
-	TConstArrayView<uint8> ModelBuffer = TConstArrayView<uint8>(&(ModelData->GetView().GetData()[GuidSize + VersionSize]), ModelData->GetView().Num() - GuidSize - VersionSize);
-
-	if (ModelBuffer.IsEmpty())
+	Session.Reset();
+	if (!TempDirForModelWithExternalData.IsEmpty())
 	{
-		UE_LOG(LogNNE, Error, TEXT("FModelInstanceORTDmlRDG::Init(): Input model data is empty."));
-		return false;
+		if (!IFileManager::Get().DeleteDirectory(*TempDirForModelWithExternalData, false, true))
+		{
+			UE_LOG(LogNNE, Warning, TEXT("FModelInstanceORTDmlRDG could not delete temp directy %s on model instance destruction."), *TempDirForModelWithExternalData);
+		}
 	}
+}
 
+bool FModelInstanceORTDmlRDG::Init()
+{
 #if WITH_EDITOR
 	try
 #endif // WITH_EDITOR
@@ -467,7 +566,11 @@ bool FModelInstanceORTDmlRDG::Init()
 		
 		SessionOptions->SetGraphOptimizationLevel(GetGraphOptimizationLevelForDML(true));
 
-		Session = MakeUnique<Ort::Session>(Environment->GetOrtEnv(), ModelBuffer.GetData(), ModelBuffer.Num(), *SessionOptions);
+		if (!Detail::CreateSession(ModelData->GetView(), *SessionOptions, *Environment, Session, TempDirForModelWithExternalData))
+		{
+			UE_LOG(LogNNE, Error, TEXT("FModelInstanceORTDmlRDG::Init(): Session creation failed."));
+			return false;
+		}
 
 		if (!ConfigureTensors(*Session))
 		{
@@ -660,16 +763,15 @@ FModelInstanceORTDmlRDG::ESetInputTensorShapesStatus FModelInstanceORTDmlRDG::Se
 		}
 	}
 
-	constexpr int32 GuidSize = sizeof(UNNERuntimeORTDml::GUID);
-	constexpr int32 VersionSize = sizeof(UNNERuntimeORTDml::Version);
-	TConstArrayView<uint8> ModelBuffer = TConstArrayView<uint8>(&(ModelData->GetView().GetData()[GuidSize + VersionSize]), ModelData->GetView().Num() - GuidSize - VersionSize);
-	check(!ModelBuffer.IsEmpty());
-
 #if WITH_EDITOR
 	try
 #endif // WITH_EDITOR
 	{
-		Session = MakeUnique<Ort::Session>(Environment->GetOrtEnv(), ModelBuffer.GetData(), ModelBuffer.Num(), *SessionOptions);
+		if (!Detail::CreateSession(ModelData->GetView(), *SessionOptions, *Environment, Session, TempDirForModelWithExternalData))
+		{
+			UE_LOG(LogNNE, Error, TEXT("Failed to recreate session!"));
+			return ESetInputTensorShapesStatus::Fail;
+		}
 	}
 #if WITH_EDITOR
 	catch (const Ort::Exception& Exception)
diff --git a/Engine/Plugins/NNE/NNERuntimeORT/Source/NNERuntimeORT/Private/NNERuntimeORTModel.h b/Engine/Plugins/NNE/NNERuntimeORT/Source/NNERuntimeORT/Private/NNERuntimeORTModel.h
--- a/Engine/Plugins/NNE/NNERuntimeORT/Source/NNERuntimeORT/Private/NNERuntimeORTModel.h
+++ b/Engine/Plugins/NNE/NNERuntimeORT/Source/NNERuntimeORT/Private/NNERuntimeORTModel.h
@@ -27,7 +27,7 @@ class FModelInstanceORTBase : public NNE::Internal::FModelInstanceBase<ModelInte
 
 public:
 	FModelInstanceORTBase(const FRuntimeConf& InRuntimeConf, TSharedRef<FEnvironment> InEnvironment);
-	virtual ~FModelInstanceORTBase() = default;
+	virtual ~FModelInstanceORTBase();
 
 	virtual typename ModelInterface::ESetInputTensorShapesStatus SetInputTensorShapes(TConstArrayView<NNE::FTensorShape> InInputShapes) override;
 
@@ -40,6 +40,7 @@ protected:
 	bool ConfigureTensors(const bool InIsInput);
 
 	FRuntimeConf RuntimeConf;
+	FString TempDirForModelWithExternalData;
 
 	/** ORT-related variables */
 	TSharedRef<FEnvironment> Environment;
@@ -131,7 +132,7 @@ public:
 
 	FModelInstanceORTDmlRDG(TSharedRef<UE::NNE::FSharedModelData> InModelData, const FRuntimeConf& InRuntimeConf, TSharedRef<FEnvironment> InEnvironment);
 
-	virtual ~FModelInstanceORTDmlRDG() = default;
+	virtual ~FModelInstanceORTDmlRDG();
 
 	virtual ESetInputTensorShapesStatus SetInputTensorShapes(TConstArrayView<NNE::FTensorShape> InInputShapes) override;
 
@@ -145,6 +146,7 @@ protected:
 
 	TSharedRef<UE::NNE::FSharedModelData> ModelData;
 	FRuntimeConf RuntimeConf;
+	FString TempDirForModelWithExternalData;
 
 	/** ORT-related variables */
 	TSharedRef<FEnvironment> Environment;
diff --git a/Engine/Plugins/NNE/NNERuntimeORT/Source/NNERuntimeORT/Private/NNERuntimeORTModelFormat.h b/Engine/Plugins/NNE/NNERuntimeORT/Source/NNERuntimeORT/Private/NNERuntimeORTModelFormat.h
new file mode 100644
--- /dev/null
+++ b/Engine/Plugins/NNE/NNERuntimeORT/Source/NNERuntimeORT/Private/NNERuntimeORTModelFormat.h
@@ -0,0 +1,39 @@
+// Copyright Epic Games, Inc. All Rights Reserved.
+
+#pragma once
+
+#include "CoreMinimal.h"
+#include "Serialization/Archive.h"
+
+namespace UE::NNERuntimeORT::Private
+{
+	struct FOnnxAdditionalDataDescriptor
+	{
+		FString Path;
+		int64 Offset;
+		int64 Size;
+	};
+
+	FORCEINLINE FArchive& operator <<(FArchive& Ar, FOnnxAdditionalDataDescriptor& AdditionalDataDescriptor)
+	{
+		Ar << AdditionalDataDescriptor.Path;
+		Ar << AdditionalDataDescriptor.Offset;
+		Ar << AdditionalDataDescriptor.Size;
+		return Ar;
+	}
+
+	struct FOnnxDataDescriptor
+	{
+
+		int64 OnnxModelDataSize = 0;
+		TArray<FOnnxAdditionalDataDescriptor> AdditionalDataDescriptors;
+	};
+
+	FORCEINLINE FArchive& operator <<(FArchive& Ar, FOnnxDataDescriptor& OnnxDataDescriptor)
+	{
+		Ar << OnnxDataDescriptor.OnnxModelDataSize;
+		Ar << OnnxDataDescriptor.AdditionalDataDescriptors;
+		return Ar;
+	}
+
+} // namespace UE::NNERuntimeORT::Private
\ No newline at end of file
diff --git a/Engine/Source/Editor/NNEEditor/NNEEditor.Build.cs b/Engine/Source/Editor/NNEEditor/NNEEditor.Build.cs
--- a/Engine/Source/Editor/NNEEditor/NNEEditor.Build.cs
+++ b/Engine/Source/Editor/NNEEditor/NNEEditor.Build.cs
@@ -1,5 +1,6 @@
 // Copyright Epic Games, Inc. All Rights Reserved.
 
+using System.IO;
 using UnrealBuildTool;
 
 public class NNEEditor : ModuleRules
@@ -26,5 +27,20 @@ public class NNEEditor : ModuleRules
 				"NNE"
 			}
 		);
+
+		if (Target.Platform == UnrealTargetPlatform.Win64)
+		{
+			string SharedLibPath = Path.Combine(ModuleDirectory, "Bin", Target.Platform.ToString());
+			string SharedLibFileName = "NNEEditorOnnxTools.dll";
+
+			PublicDelayLoadDLLs.Add(SharedLibFileName);
+
+			RuntimeDependencies.Add("$(TargetOutputDir)/"+ SharedLibFileName, Path.Combine(SharedLibPath, SharedLibFileName));
+
+			PublicDefinitions.Add("NNEEDITORONNXTOOLS_SUPPORTED");
+			PublicDefinitions.Add("NNEEDITORONNXTOOLS_SHAREDLIB_FILENAME=" + SharedLibFileName);
+		}
+
+		PublicDefinitions.Add("UE_NNEEDITORONNXTOOLS");
 	}
 }
diff --git a/Engine/Source/Editor/NNEEditor/Private/NNEEditorModelDataFactory.cpp b/Engine/Source/Editor/NNEEditor/Private/NNEEditorModelDataFactory.cpp
--- a/Engine/Source/Editor/NNEEditor/Private/NNEEditorModelDataFactory.cpp
+++ b/Engine/Source/Editor/NNEEditor/Private/NNEEditorModelDataFactory.cpp
@@ -5,8 +5,9 @@
 #include "Editor.h"
 #include "EngineAnalytics.h"
 #include "Kismet/GameplayStatics.h"
-#include "NNE.h"
+#include "NNEEditorOnnxFileLoaderHelper.h"
 #include "NNEModelData.h"
+#include "Misc/Paths.h"
 #include "Subsystems/ImportSubsystem.h"
 
 UNNEModelDataFactory::UNNEModelDataFactory(const FObjectInitializer& ObjectInitializer) : UFactory(ObjectInitializer)
@@ -18,20 +19,25 @@ UNNEModelDataFactory::UNNEModelDataFactory(const FObjectInitializer& ObjectIniti
 	Formats.Add("onnx;Open Neural Network Exchange Format");
 }
 
-UObject* UNNEModelDataFactory::FactoryCreateBinary(UClass* Class, UObject* InParent, FName Name, EObjectFlags Flags, UObject* Context, const TCHAR * Type, const uint8 *& Buffer, const uint8 * BufferEnd, FFeedbackContext* Warn)
+
+UObject* UNNEModelDataFactory::FactoryCreateFile(UClass* InClass, UObject* InParent, FName InName, EObjectFlags Flags, const FString& Filename, const TCHAR* Parms, FFeedbackContext* Warn, bool& bOutOperationCanceled)
 {
-	GEditor->GetEditorSubsystem<UImportSubsystem>()->BroadcastAssetPreImport(this, Class, InParent, Name, Type);
+	AdditionalImportedObjects.Empty();
+
+	FString FileExtension = FPaths::GetExtension(Filename);
 
-	if (!Type || !Buffer || !BufferEnd || BufferEnd - Buffer <= 0)
+	GEditor->GetEditorSubsystem<UImportSubsystem>()->BroadcastAssetPreImport(this, InClass, InParent, InName, *FileExtension);
+	
+	int64 ModelFileSize = 0;
+	UNNEModelData* ModelData = NewObject<UNNEModelData>(InParent, InClass, InName, Flags);
+	check(ModelData)
+
+	if (!UE::NNEEditor::Internal::OnnxFileLoaderHelper::InitUNNEModelDataFromFile(*ModelData, ModelFileSize, Filename))
 	{
 		GEditor->GetEditorSubsystem<UImportSubsystem>()->BroadcastAssetPostImport(this, nullptr);
 		return nullptr;
 	}
 
-	UNNEModelData* ModelData = NewObject<UNNEModelData>(InParent, Class, Name, Flags);
-	TConstArrayView<uint8> BufferView = MakeArrayView(Buffer, BufferEnd-Buffer);
-	ModelData->Init(Type, BufferView);
-
 	GEditor->GetEditorSubsystem<UImportSubsystem>()->BroadcastAssetPostImport(this, ModelData);
 
 	if (FEngineAnalytics::IsAvailable())
@@ -39,7 +45,7 @@ UObject* UNNEModelDataFactory::FactoryCreateBinary(UClass* Class, UObject* InPar
 		TArray<FAnalyticsEventAttribute> Attributes = MakeAnalyticsEventAttributeArray(
 			TEXT("PlatformName"), UGameplayStatics::GetPlatformName(),
 			TEXT("FactoryName"), TEXT("UNNEModelDataFactory"),
-			TEXT("ModelFileSize"), BufferView.Num()
+			TEXT("ModelFileSize"), ModelFileSize
 		);
 		FEngineAnalytics::GetProvider().RecordEvent(TEXT("NeuralNetworkEngine.FactoryCreateBinary"), Attributes);
 	}
diff --git a/Engine/Source/Editor/NNEEditor/Private/NNEEditorModelDataFactory.h b/Engine/Source/Editor/NNEEditor/Private/NNEEditorModelDataFactory.h
--- a/Engine/Source/Editor/NNEEditor/Private/NNEEditorModelDataFactory.h
+++ b/Engine/Source/Editor/NNEEditor/Private/NNEEditorModelDataFactory.h
@@ -22,7 +22,7 @@ public:
 
 public:
 	//~ Begin UFactory Interface
-	virtual UObject* FactoryCreateBinary(UClass* Class, UObject* InParent, FName Name, EObjectFlags Flags, UObject* Context, const TCHAR* Type, const uint8*& Buffer, const uint8* BufferEnd, FFeedbackContext* Warn) override;
+	virtual UObject* FactoryCreateFile(UClass* InClass, UObject* InParent, FName InName, EObjectFlags Flags, const FString& Filename, const TCHAR* Parms, FFeedbackContext* Warn, bool& bOutOperationCanceled) override;
 	virtual bool FactoryCanImport(const FString& Filename) override;
 	//~ End UFactory Interface
 };
\ No newline at end of file
diff --git a/Engine/Source/Editor/NNEEditor/Private/NNEEditorModule.cpp b/Engine/Source/Editor/NNEEditor/Private/NNEEditorModule.cpp
--- a/Engine/Source/Editor/NNEEditor/Private/NNEEditorModule.cpp
+++ b/Engine/Source/Editor/NNEEditor/Private/NNEEditorModule.cpp
@@ -1,13 +1,49 @@
 // Copyright Epic Games, Inc. All Rights Reserved.
 
 #include "Modules/ModuleManager.h"
+#include "HAL/FileManager.h"
+#include "HAL/PlatformProcess.h"
+#include "Misc/Paths.h"
+#include "NNE.h"
 #include "NNEEditorModelDataActions.h"
+#include "NNEEditorOnnxModelInspector.h"
+
+namespace UE::NNEEditor::Private::OnnxHelper
+{
+	void* GetSharedLibHandle(const FString& SharedLibPath)
+	{
+		if (!FPaths::FileExists(SharedLibPath))
+		{
+			UE_LOG(LogNNE, Error, TEXT("Failed to find the NNEEditorOnnxTools shared library %s."), *SharedLibPath);
+			return nullptr;
+		}
+
+		void* SharedLibHandle = FPlatformProcess::GetDllHandle(*SharedLibPath);
+
+		if (!SharedLibHandle)
+		{
+			UE_LOG(LogNNE, Error, TEXT("Failed to load the NNEEditorOnnxTools shared library %s."), *SharedLibPath);
+		}
+
+		return SharedLibHandle;
+	}
+
+} // namespace UE::NNEEditor::Private::OnnxHelper
 
 class FNNEEditorModule : public IModuleInterface
 {
+	void* NNEEditorOnnxToolsSharedLibHandle = nullptr;
+
 public:
 	virtual void StartupModule() override
 	{
+#ifdef NNEEDITORONNXTOOLS_SUPPORTED
+		const FString ModuleDir = FPlatformProcess::GetModulesDirectory();
+		const FString NNEEditorOnnxToolsSharedLibPath = FPaths::Combine(ModuleDir, TEXT(PREPROCESSOR_TO_STRING(NNEEDITORONNXTOOLS_SHAREDLIB_FILENAME)));
+		NNEEditorOnnxToolsSharedLibHandle = UE::NNEEditor::Private::OnnxHelper::GetSharedLibHandle(NNEEditorOnnxToolsSharedLibPath);
+		UE::NNEEditor::Private::OnnxModelInspectorHelper::SetupSharedLibFunctionPointer(NNEEditorOnnxToolsSharedLibHandle);
+#endif // ifdef NNEEDITORONNXTOOLS_SUPPORTED
+
 		ModelDataAssetTypeActions = MakeShared<UE::NNEEditor::Private::FModelDataAssetTypeActions>();
 		FAssetToolsModule::GetModule().Get().RegisterAssetTypeActions(ModelDataAssetTypeActions.ToSharedRef());
 	}
@@ -18,6 +54,14 @@ public:
 		{
 			FAssetToolsModule::GetModule().Get().UnregisterAssetTypeActions(ModelDataAssetTypeActions.ToSharedRef());
 		}
+
+		UE::NNEEditor::Private::OnnxModelInspectorHelper::ClearSharedLibFunctionPointer();
+
+		if (NNEEditorOnnxToolsSharedLibHandle != nullptr)
+		{
+			FPlatformProcess::FreeDllHandle(NNEEditorOnnxToolsSharedLibHandle);
+			NNEEditorOnnxToolsSharedLibHandle = nullptr;
+		}
 	}
 
 private:
diff --git a/Engine/Source/Editor/NNEEditor/Private/NNEEditorOnnxFileLoaderHelper.cpp b/Engine/Source/Editor/NNEEditor/Private/NNEEditorOnnxFileLoaderHelper.cpp
new file mode 100644
--- /dev/null
+++ b/Engine/Source/Editor/NNEEditor/Private/NNEEditorOnnxFileLoaderHelper.cpp
@@ -0,0 +1,134 @@
+// Copyright Epic Games, Inc. All Rights Reserved.
+
+#include "NNEEditorOnnxFileLoaderHelper.h"
+
+#include "CoreGlobals.h"
+#include "HAL/FileManager.h"
+#include "NNE.h"
+#include "NNEEditorOnnxModelInspector.h"
+#include "NNEModelData.h"
+#include "Misc/FileHelper.h"
+#include "Misc/Paths.h"
+#include "Serialization/Archive.h"
+#include "Serialization/MemoryWriter.h"
+
+namespace UE::NNEEditor::Internal::OnnxFileLoaderHelper
+{
+	namespace Details
+	{
+		bool AppendFileToArray(const FString& Filename, TArray64<uint8>& Buffer)
+		{
+			FScopedLoadingState ScopedLoadingState(*Filename);
+
+			FArchive* Reader = IFileManager::Get().CreateFileReader(*Filename, 0);
+			if (!Reader)
+			{
+				return false;
+			}
+
+			int64 FileSize = Reader->TotalSize();
+			if (FileSize > 0)
+			{
+				int64 BufferSizeBefore = Buffer.Num();
+
+				Buffer.AddUninitialized(FileSize);
+				Reader->Serialize(Buffer.GetData()+BufferSizeBefore, FileSize);
+			}
+			
+			bool Success = Reader->Close();
+			delete Reader;
+			return Success;
+		}
+	}
+
+	bool InitUNNEModelDataFromFile(UNNEModelData& ModelData, int64& ModelFileSize, const FString& Filename)
+	{
+		FString FileExtension = FPaths::GetExtension(Filename);
+		ModelFileSize = 0;
+
+		TArray64<uint8> OnnxData;
+		if (!FFileHelper::LoadFileToArray(OnnxData, *Filename))
+		{
+			UE_LOG(LogNNE, Error, TEXT("Failed to load file '%s'"), *Filename);
+			return false;
+		}
+
+		if (!*FileExtension || OnnxData.IsEmpty())
+		{
+			return false;
+		}
+
+		TSet<FString> ExternalDataRelativeFilePaths;
+
+		//If platform does not support parsing the onnx file, we import assuming the model does not use external data.
+		if (UE::NNEEditor::Private::OnnxModelInspectorHelper::IsSharedLibFunctionPointerSetup())
+		{
+			if (!UE::NNEEditor::Private::OnnxModelInspectorHelper::GetExternalDataFilePaths(OnnxData, ExternalDataRelativeFilePaths))
+			{
+				UE_LOG(LogNNE, Warning, TEXT("Could not parse the input model as ONNX ModelProto, model external data won't be imported if any."));
+			}
+		}
+
+		ModelFileSize = OnnxData.Num();
+
+		TMap<FString, TConstArrayView<uint8>> AdditionalBuffers;
+		TArray64<uint8> OnnxExternalDataBytesBuffer;
+		TArray64<uint8> OnnxExternalDataDescriptorBuffer;
+
+		if (!ExternalDataRelativeFilePaths.IsEmpty())
+		{
+			FString BasePath = FPaths::GetPath(*Filename);
+			TMap<FString, int64> OnnxExternalDataDescriptor;
+			int64 TotalFileSize = 0;
+
+			for (const FString& DataRelativeFilePath : ExternalDataRelativeFilePaths)
+			{
+				FString DataFilePath = FPaths::Combine(BasePath, DataRelativeFilePath);
+				int64 FileSize = IFileManager::Get().FileSize(*DataFilePath);
+				if (FileSize < 0)
+				{
+					UE_LOG(LogNNE, Error, TEXT("Failed to find file size for external data file '%s' for Onnx Model '%s'"), *DataFilePath, *Filename);
+					return false;
+				}
+
+				OnnxExternalDataDescriptor.Emplace(DataRelativeFilePath, FileSize);
+				TotalFileSize += FileSize;
+			}
+
+			OnnxExternalDataBytesBuffer.Reset(TotalFileSize);
+			for (const FString& DataRelativeFilePath : ExternalDataRelativeFilePaths)
+			{
+				FString DataFilePath = FPaths::Combine(BasePath, DataRelativeFilePath);
+
+				if (!Details::AppendFileToArray(DataFilePath, OnnxExternalDataBytesBuffer))
+				{
+					UE_LOG(LogNNE, Error, TEXT("Failed to load external data file '%s' for Onnx Model '%s'"), *DataFilePath, *Filename);
+					return false;
+				}
+
+				//Uncomment to trace loading of external onnx data from files.
+				//UE_LOG(LogNNE, Display, TEXT("Loaded external data file '%s' while importing Onnx Model '%s'."), *DataFilePath, *Filename);
+			}
+			check(OnnxExternalDataBytesBuffer.Num() == TotalFileSize);
+
+			//Should be kept in sync with UE::NNERuntimeORT::Private::Details::OnnxExternalDataDescriptorKey & OnnxExternalDataBytesKey in NNERuntimeORT.cpp
+			static FString OnnxExternalDataDescriptorKey(TEXT("OnnxExternalDataDescriptor"));
+			static FString OnnxExternalDataBytesKey(TEXT("OnnxExternalDataBytes"));
+			
+			FMemoryWriter64 OnnxExternalDataAggregatedDescriptorWriter(OnnxExternalDataDescriptorBuffer, /*bIsPersistent = */true);
+
+			OnnxExternalDataAggregatedDescriptorWriter << OnnxExternalDataDescriptor;
+
+			AdditionalBuffers.Emplace(OnnxExternalDataDescriptorKey, OnnxExternalDataDescriptorBuffer);
+			AdditionalBuffers.Emplace(OnnxExternalDataBytesKey, OnnxExternalDataBytesBuffer);
+
+			ModelFileSize += OnnxExternalDataBytesBuffer.Num();
+		}
+
+		ModelData.Init(*FileExtension, OnnxData, AdditionalBuffers);
+
+		return true;
+	}
+
+} // namespace UE::NNEEditor::Internal::OnnxFileLoaderHelper
+
diff --git a/Engine/Source/Editor/NNEEditor/Private/NNEEditorOnnxFileLoaderHelper.h b/Engine/Source/Editor/NNEEditor/Private/NNEEditorOnnxFileLoaderHelper.h
new file mode 100644
--- /dev/null
+++ b/Engine/Source/Editor/NNEEditor/Private/NNEEditorOnnxFileLoaderHelper.h
@@ -0,0 +1,12 @@
+// Copyright Epic Games, Inc. All Rights Reserved.
+#pragma once
+
+#include "CoreMinimal.h"
+
+class UNNEModelData;
+
+namespace UE::NNEEditor::Internal::OnnxFileLoaderHelper
+{
+	bool NNEEDITOR_API InitUNNEModelDataFromFile(UNNEModelData& ModelData, int64& ModelFileSize, const FString& Filename);
+
+} // namespace UE::NNEEditor::Internal::OnnxFileLoaderHelper
diff --git a/Engine/Source/Editor/NNEEditor/Private/NNEEditorOnnxModelInspector.cpp b/Engine/Source/Editor/NNEEditor/Private/NNEEditorOnnxModelInspector.cpp
new file mode 100644
--- /dev/null
+++ b/Engine/Source/Editor/NNEEditor/Private/NNEEditorOnnxModelInspector.cpp
@@ -0,0 +1,113 @@
+// Copyright Epic Games, Inc. All Rights Reserved.
+
+#include "NNEEditorOnnxModelInspector.h"
+
+#include "NNE.h"
+#include "NNEEditorOnnxTools.h"
+
+namespace UE::NNEEditor::Private
+{
+namespace OnnxModelInspectorHelper
+{
+	class FNNEOnnxToolsWrapper
+	{
+		typedef NNEEditorOnnxTools_Status(*CreateExternalDataDescriptorFn)(const void* InData, const int Size, NNEEditorOnnxTools_ExternalDataDescriptor** Descriptor);
+		typedef void(*ReleaseExternalDataDescriptorFn)(NNEEditorOnnxTools_ExternalDataDescriptor** Descriptor);
+		typedef const char* (*GetNextExternalDataPathFn)(NNEEditorOnnxTools_ExternalDataDescriptor* Descriptor);
+
+		CreateExternalDataDescriptorFn CreateExternalDataDescriptor{};
+		ReleaseExternalDataDescriptorFn ReleaseExternalDataDescriptor{};
+		GetNextExternalDataPathFn GetNextExternalDataPath{};
+
+	public:
+		static FNNEOnnxToolsWrapper& GetInstance()
+		{
+			static FNNEOnnxToolsWrapper Instance;
+			return Instance;
+		}
+
+		void SetupSharedLibFunctionPointer(void* SharedLibHandle)
+		{
+			if (SharedLibHandle)
+			{
+				CreateExternalDataDescriptor = reinterpret_cast<CreateExternalDataDescriptorFn>(FPlatformProcess::GetDllExport(SharedLibHandle, TEXT("NNEEditorOnnxTools_CreateExternalDataDescriptor")));
+				ReleaseExternalDataDescriptor = reinterpret_cast<ReleaseExternalDataDescriptorFn>(FPlatformProcess::GetDllExport(SharedLibHandle, TEXT("NNEEditorOnnxTools_ReleaseExternalDataDescriptor")));
+				GetNextExternalDataPath = reinterpret_cast<GetNextExternalDataPathFn>(FPlatformProcess::GetDllExport(SharedLibHandle, TEXT("NNEEditorOnnxTools_GetNextExternalDataPath")));
+			}
+			if (CreateExternalDataDescriptor == nullptr ||
+				ReleaseExternalDataDescriptor == nullptr ||
+				GetNextExternalDataPath == nullptr)
+			{
+				UE_LOG(LogNNE, Warning, TEXT("Could not find required function pointers in NNEEditorOnnxTools shared library."));
+			}
+		}
+
+		void ClearSharedLibFunctionPointer()
+		{
+			CreateExternalDataDescriptor = nullptr;
+			ReleaseExternalDataDescriptor = nullptr;
+			GetNextExternalDataPath = nullptr;
+		}
+
+		bool IsSharedLibFunctionPointerSetup()
+		{
+			return (CreateExternalDataDescriptor != nullptr &&
+					ReleaseExternalDataDescriptor != nullptr &&
+					GetNextExternalDataPath != nullptr);
+		}
+
+		bool GetExternalDataFilePaths(TConstArrayView<uint8> ONNXData, TSet<FString>& ExternalDataFilePaths)
+		{
+			ExternalDataFilePaths.Reset();
+
+			if (!IsSharedLibFunctionPointerSetup())
+				return false;
+
+			NNEEditorOnnxTools_ExternalDataDescriptor* ExternalDataDescriptor = nullptr;
+			const NNEEditorOnnxTools_Status Result = CreateExternalDataDescriptor(ONNXData.GetData(), ONNXData.Num(), &ExternalDataDescriptor);
+
+			if (Result == NNEEditorOnnxTools_Status::Ok)
+			{
+				const char* ExternalPathFromLib = GetNextExternalDataPath(ExternalDataDescriptor);
+				while (ExternalPathFromLib != nullptr)
+				{
+					const FString ExternalPath = StringCast<TCHAR>(ExternalPathFromLib).Get();
+					ExternalDataFilePaths.Emplace(ExternalPath);
+					ExternalPathFromLib = GetNextExternalDataPath(ExternalDataDescriptor);
+				}
+
+				ReleaseExternalDataDescriptor(&ExternalDataDescriptor);
+			}
+			else
+			{
+				return false;
+			}
+
+			return true;
+
+		}
+	};
+
+	void SetupSharedLibFunctionPointer(void* SharedLibHandle)
+	{
+		FNNEOnnxToolsWrapper::GetInstance().SetupSharedLibFunctionPointer(SharedLibHandle);
+	}
+
+	void ClearSharedLibFunctionPointer()
+	{
+		FNNEOnnxToolsWrapper::GetInstance().ClearSharedLibFunctionPointer();
+	}
+
+	bool IsSharedLibFunctionPointerSetup()
+	{
+		return FNNEOnnxToolsWrapper::GetInstance().IsSharedLibFunctionPointerSetup();
+	}
+
+	bool GetExternalDataFilePaths(TConstArrayView<uint8> ONNXData, TSet<FString>& ExternalDataFilePaths)
+	{
+		return FNNEOnnxToolsWrapper::GetInstance().GetExternalDataFilePaths(ONNXData, ExternalDataFilePaths);
+	}
+
+} // namespace OnnxModelInspectorHelper
+
+} // namespace UE::NNEEditor::Private
diff --git a/Engine/Source/Editor/NNEEditor/Private/NNEEditorOnnxModelInspector.h b/Engine/Source/Editor/NNEEditor/Private/NNEEditorOnnxModelInspector.h
new file mode 100644
--- /dev/null
+++ b/Engine/Source/Editor/NNEEditor/Private/NNEEditorOnnxModelInspector.h
@@ -0,0 +1,13 @@
+// Copyright Epic Games, Inc. All Rights Reserved.
+#pragma once
+
+#include "CoreMinimal.h"
+
+namespace UE::NNEEditor::Private::OnnxModelInspectorHelper
+{
+	void SetupSharedLibFunctionPointer(void* SharedLibHandle);
+	void ClearSharedLibFunctionPointer();
+	bool IsSharedLibFunctionPointerSetup();
+	bool GetExternalDataFilePaths(TConstArrayView<uint8> ONNXData, TSet<FString>& ExternalDataFilePaths);
+
+} // namespace UE::NNEEditor::Private::OnnxModelInspectorHelper
diff --git a/Engine/Source/Editor/NNEEditor/Private/NNEEditorOnnxTools.h b/Engine/Source/Editor/NNEEditor/Private/NNEEditorOnnxTools.h
new file mode 100644
--- /dev/null
+++ b/Engine/Source/Editor/NNEEditor/Private/NNEEditorOnnxTools.h
@@ -0,0 +1,33 @@
+#pragma once
+
+#ifdef NNEEDITORONNXTOOLS_EXPORT
+#define NNEEDITORONNXTOOLS_API __declspec(dllexport)
+#else
+#define NNEEDITORONNXTOOLS_API
+#endif
+
+#ifndef UE_NNEEDITORONNXTOOLS
+#include <cstdint>
+typedef uint8_t uint8;
+#else
+#include "CoreMinimal.h"
+#endif
+
+enum class NNEEditorOnnxTools_Status : uint8
+{
+	Ok = 0,
+	Fail_CannotParseAsModelProto
+};
+
+struct NNEEditorOnnxTools_ExternalDataDescriptor;
+
+extern "C"
+{
+
+NNEEDITORONNXTOOLS_API NNEEditorOnnxTools_Status NNEEditorOnnxTools_CreateExternalDataDescriptor(const void* InData, const int Size, NNEEditorOnnxTools_ExternalDataDescriptor** Descriptor);
+
+NNEEDITORONNXTOOLS_API void NNEEditorOnnxTools_ReleaseExternalDataDescriptor(NNEEditorOnnxTools_ExternalDataDescriptor** Descriptor);
+
+NNEEDITORONNXTOOLS_API const char* NNEEditorOnnxTools_GetNextExternalDataPath(NNEEditorOnnxTools_ExternalDataDescriptor* Descriptor);
+
+}
\ No newline at end of file
