NNE Fix following CL 35820991 in MakeOnnxDataDescriptor

#tests nne smoke test win64 editor

diff --git a/Engine/Plugins/NNE/NNERuntimeORT/Source/NNERuntimeORT/Private/NNERuntimeORT.cpp b/Engine/Plugins/NNE/NNERuntimeORT/Source/NNERuntimeORT/Private/NNERuntimeORT.cpp
--- a/Engine/Plugins/NNE/NNERuntimeORT/Source/NNERuntimeORT/Private/NNERuntimeORT.cpp
+++ b/Engine/Plugins/NNE/NNERuntimeORT/Source/NNERuntimeORT/Private/NNERuntimeORT.cpp
@@ -26,7 +26,7 @@ FGuid UNNERuntimeORTDml::GUID = FGuid((int32)'O', (int32)'D', (int32)'M', (int32
 int32 UNNERuntimeORTDml::Version = 0x00000003;
 
 namespace UE::NNERuntimeORT::Private::Details
-{
+{ 
 	//Should be kept in sync with OnnxFileLoaderHelper::InitUNNEModelDataFromFile()
 	static FString OnnxExternalDataDescriptorKey(TEXT("OnnxExternalDataDescriptor"));
 	static FString OnnxExternalDataBytesKey(TEXT("OnnxExternalDataBytes"));
@@ -48,10 +48,12 @@ namespace UE::NNERuntimeORT::Private::Details
 			for (const auto& Element : ExternalDataSizes)
 			{
 				const FString DataFilePath = Element.Key;
-				OnnxDataDescriptor.AdditionalDataDescriptors.Emplace(
-					DataFilePath,
-					CurrentBucketOffset,
-					Element.Value);
+				FOnnxAdditionalDataDescriptor DataDescriptor;
+				DataDescriptor.Path = DataFilePath;
+				DataDescriptor.Offset = CurrentBucketOffset;
+				DataDescriptor.Size = Element.Value;
+
+				OnnxDataDescriptor.AdditionalDataDescriptors.Emplace(DataDescriptor);
 				CurrentBucketOffset += Element.Value;
 			}
 		}
diff --git a/Engine/Source/Editor/NNEEditor/Private/NNEEditorOnnxTools.h b/Engine/Source/Editor/NNEEditor/Private/NNEEditorOnnxTools.h
--- a/Engine/Source/Editor/NNEEditor/Private/NNEEditorOnnxTools.h
+++ b/Engine/Source/Editor/NNEEditor/Private/NNEEditorOnnxTools.h
@@ -1,3 +1,5 @@
+// Copyright Epic Games, Inc. All Rights Reserved.
+
 #pragma once
 
 #ifdef NNEEDITORONNXTOOLS_EXPORT
