[NNE] NNE runtimes plugin description and category revision/harmonization.

#rnx

diff --git a/Engine/Plugins/Experimental/NNERuntimeIREE/NNERuntimeIREE.uplugin b/Engine/Plugins/Experimental/NNERuntimeIREE/NNERuntimeIREE.uplugin
--- a/Engine/Plugins/Experimental/NNERuntimeIREE/NNERuntimeIREE.uplugin
+++ b/Engine/Plugins/Experimental/NNERuntimeIREE/NNERuntimeIREE.uplugin
@@ -3,8 +3,8 @@
 	"Version": 0.1,
 	"VersionName": "0.1",
 	"FriendlyName": "NNERuntimeIREE",
-	"Description": "A runtime implementing the NNE API which is based on IREE, MLIR and LLVM and compiles neural networks directly to game code.",
-	"Category": "Other",
+	"Description": "A runtime implementing the Neural Network Engine (NNE) API which is based on IREE, MLIR and LLVM and compiles neural networks directly to game code.",
+	"Category": "ML",
 	"CreatedBy": "Epic Games, Inc.",
 	"CreatedByURL": "https://epicgames.com",
 	"DocsURL": "",
diff --git a/Engine/Plugins/Experimental/NNERuntimeRDG/NNERuntimeRDG.uplugin b/Engine/Plugins/Experimental/NNERuntimeRDG/NNERuntimeRDG.uplugin
--- a/Engine/Plugins/Experimental/NNERuntimeRDG/NNERuntimeRDG.uplugin
+++ b/Engine/Plugins/Experimental/NNERuntimeRDG/NNERuntimeRDG.uplugin
@@ -3,8 +3,8 @@
 	"Version": 0.1,
 	"VersionName": "0.1",
 	"FriendlyName": "NNERuntimeRDG",
-	"Description": "GPU RDG accelerated runtime for the NNE plugin (deep learning and neural network inference in Unreal Engine) backed by Unreal rendergraph RDG.",
-	"Category": "Experimental",
+	"Description": "A runtime implementing the Neural Network Engine (NNE) API, using the Render Dependency Graph (RDG).",
+	"Category": "ML",
 	"CreatedBy": "Epic Games, Inc.",
 	"CreatedByURL": "https://epicgames.com",
 	"DocsURL": "",
diff --git a/Engine/Plugins/NNE/NNERuntimeORT/NNERuntimeORT.uplugin b/Engine/Plugins/NNE/NNERuntimeORT/NNERuntimeORT.uplugin
--- a/Engine/Plugins/NNE/NNERuntimeORT/NNERuntimeORT.uplugin
+++ b/Engine/Plugins/NNE/NNERuntimeORT/NNERuntimeORT.uplugin
@@ -1,10 +1,10 @@
 {
 	"FileVersion": 3,
 	"Version": 1,
-	"VersionName": "1.0.0",
+	"VersionName": "1.0",
 	"FriendlyName": "NNERuntimeORT",
-	"Description": "ONNXRuntime accelerated runtimes for the NNE plugin (deep learning and neural network inference in Unreal Engine) backed by the Cpu and Dml execution providers.",
-	"Category": "AI",
+	"Description": "ONNX Runtime backed runtime for the Neural Network Engine (NNE), accelerated by the CPU and DirectML execution providers.",
+	"Category": "ML",
 	"CreatedBy": "Epic Games, Inc.",
 	"CreatedByURL": "https://epicgames.com",
 	"DocsURL": "https://dev.epicgames.com/community/learning/courses/e7w/unreal-engine-neural-network-engine-nne",
