[NNE] Remove analytics as we are covered by tableau

Problem:
* We are using custom analytics however related state are already well covered by tableau per plugins stats

Solution:
* Remove custom stats

Tests:
* Win64 Editor debug : smoke tests
* Win64 Standalone packaged debug : smoke tests

#rnx

diff --git a/Engine/Plugins/Experimental/NNERuntimeIREE/Source/NNERuntimeIREE/Private/NNERuntimeIREE.cpp b/Engine/Plugins/Experimental/NNERuntimeIREE/Source/NNERuntimeIREE/Private/NNERuntimeIREE.cpp
--- a/Engine/Plugins/Experimental/NNERuntimeIREE/Source/NNERuntimeIREE/Private/NNERuntimeIREE.cpp
+++ b/Engine/Plugins/Experimental/NNERuntimeIREE/Source/NNERuntimeIREE/Private/NNERuntimeIREE.cpp
@@ -11,7 +11,6 @@
 #include "Misc/FileHelper.h"
 #endif // WITH_EDITOR
 
-#include "EngineAnalytics.h"
 #include "HAL/Platform.h"
 #include "Interfaces/ITargetPlatform.h"
 #include "Kismet/GameplayStatics.h"
@@ -343,16 +342,6 @@ TSharedPtr<UE::NNE::IModelCPU> UNNERuntimeIREECpu::CreateModelCPU(const TObjectP
 		return TSharedPtr<UE::NNE::IModelCPU>();
 	}
 
-	if (FEngineAnalytics::IsAvailable())
-	{
-		TArray<FAnalyticsEventAttribute> Attributes = MakeAnalyticsEventAttributeArray(
-			TEXT("PlatformName"), UGameplayStatics::GetPlatformName(),
-			TEXT("HashedRuntimeName"), FMD5::HashAnsiString(*GetRuntimeName()),
-			TEXT("ModelDataSize"), SharedDataView.Num()
-		);
-		FEngineAnalytics::GetProvider().RecordEvent(TEXT("NeuralNetworkEngine.CreateModel"), Attributes);
-	}
-
 	return Model;
 }
 
@@ -428,16 +417,6 @@ TSharedPtr<UE::NNE::IModelGPU> UNNERuntimeIREEGpu::CreateModelGPU(const TObjectP
 	UE::NNE::IModelGPU* IModel = nullptr;
 	TConstArrayView<uint8> SharedDataView = ModelData->GetModelData(GetRuntimeName())->GetView();
 
-	if (FEngineAnalytics::IsAvailable())
-	{
-		TArray<FAnalyticsEventAttribute> Attributes = MakeAnalyticsEventAttributeArray(
-			TEXT("PlatformName"), UGameplayStatics::GetPlatformName(),
-			TEXT("HashedRuntimeName"), FMD5::HashAnsiString(*GetRuntimeName()),
-			TEXT("ModelDataSize"), SharedDataView.Num()
-		);
-		FEngineAnalytics::GetProvider().RecordEvent(TEXT("NeuralNetworkEngine.CreateModel"), Attributes);
-	}
-
 	return TSharedPtr<UE::NNE::IModelGPU>(IModel);
 }
 
@@ -575,16 +554,6 @@ TSharedPtr<UE::NNE::IModelRDG> UNNERuntimeIREERdg::CreateModelRDG(const TObjectP
 	UE::NNE::IModelRDG* IModel = nullptr;
 	TConstArrayView<uint8> SharedDataView = ModelData->GetModelData(GetRuntimeName())->GetView();
 
-	if (FEngineAnalytics::IsAvailable())
-	{
-		TArray<FAnalyticsEventAttribute> Attributes = MakeAnalyticsEventAttributeArray(
-			TEXT("PlatformName"), UGameplayStatics::GetPlatformName(),
-			TEXT("HashedRuntimeName"), FMD5::HashAnsiString(*GetRuntimeName()),
-			TEXT("ModelDataSize"), SharedDataView.Num()
-		);
-		FEngineAnalytics::GetProvider().RecordEvent(TEXT("NeuralNetworkEngine.CreateModel"), Attributes);
-	}
-
 	return TSharedPtr<UE::NNE::IModelRDG>(IModel);
 }
 
diff --git a/Engine/Plugins/Experimental/NNERuntimeIREE/Source/NNERuntimeIREEEditor/Private/NNERuntimeIREEModelDataFactory.cpp b/Engine/Plugins/Experimental/NNERuntimeIREE/Source/NNERuntimeIREEEditor/Private/NNERuntimeIREEModelDataFactory.cpp
--- a/Engine/Plugins/Experimental/NNERuntimeIREE/Source/NNERuntimeIREEEditor/Private/NNERuntimeIREEModelDataFactory.cpp
+++ b/Engine/Plugins/Experimental/NNERuntimeIREE/Source/NNERuntimeIREEEditor/Private/NNERuntimeIREEModelDataFactory.cpp
@@ -4,13 +4,11 @@
 
 #include "CoreMinimal.h"
 #include "Editor.h"
-#include "EngineAnalytics.h"
 #include "Framework/Application/SlateApplication.h"
 #include "Framework/Commands/UIAction.h"
 #include "Framework/MultiBox/MultiBoxBuilder.h"
 #include "Input/Reply.h"
 #include "Interfaces/IMainFrameModule.h"
-#include "Kismet/GameplayStatics.h"
 #include "Modules/ModuleManager.h"
 #include "NNE.h"
 #include "NNEModelData.h"
@@ -216,16 +214,6 @@ UObject* UNNERuntimeIREEModelDataFactory::FactoryCreateBinary(UClass* Class, UOb
 
 	GEditor->GetEditorSubsystem<UImportSubsystem>()->BroadcastAssetPostImport(this, ModelData);
 
-	if (FEngineAnalytics::IsAvailable())
-	{
-		TArray<FAnalyticsEventAttribute> Attributes = MakeAnalyticsEventAttributeArray(
-			TEXT("PlatformName"), UGameplayStatics::GetPlatformName(),
-			TEXT("FactoryName"), TEXT("UNNERuntimeIREEModelDataFactory"),
-			TEXT("ModelFileSize"), BufferView.Num()
-		);
-		FEngineAnalytics::GetProvider().RecordEvent(TEXT("NeuralNetworkEngine.FactoryCreateBinary"), Attributes);
-	}
-
 	return ModelData;
 }
 
diff --git a/Engine/Plugins/Experimental/NNERuntimeRDG/Source/NNERuntimeRDG/Private/NNERuntimeRDGHlsl.cpp b/Engine/Plugins/Experimental/NNERuntimeRDG/Source/NNERuntimeRDG/Private/NNERuntimeRDGHlsl.cpp
--- a/Engine/Plugins/Experimental/NNERuntimeRDG/Source/NNERuntimeRDG/Private/NNERuntimeRDGHlsl.cpp
+++ b/Engine/Plugins/Experimental/NNERuntimeRDG/Source/NNERuntimeRDG/Private/NNERuntimeRDGHlsl.cpp
@@ -2,8 +2,6 @@
 
 #include "NNERuntimeRDGHlsl.h"
 
-#include "EngineAnalytics.h"
-#include "Kismet/GameplayStatics.h"
 #include "Misc/SecureHash.h"
 #include "NNE.h"
 #include "NNEAttributeMap.h"
@@ -246,15 +244,5 @@ TSharedPtr<UE::NNE::IModelRDG> UNNERuntimeRDGHlslImpl::CreateModelRDG(const TObj
 	check(Data.IsValid());
 	UE::NNERuntimeRDG::Private::Hlsl::FModel* Model = new UE::NNERuntimeRDG::Private::Hlsl::FModel(Data);
 
-	if (FEngineAnalytics::IsAvailable())
-	{
-		TArray<FAnalyticsEventAttribute> Attributes = MakeAnalyticsEventAttributeArray(
-			TEXT("PlatformName"), UGameplayStatics::GetPlatformName(),
-			TEXT("HashedRuntimeName"), FMD5::HashAnsiString(*GetRuntimeName()),
-			TEXT("ModelDataSize"), Data->GetView().Num()
-		);
-		FEngineAnalytics::GetProvider().RecordEvent(TEXT("NeuralNetworkEngine.CreateModel"), Attributes);
-	}
-
 	return TSharedPtr<UE::NNE::IModelRDG>(Model);
 }
diff --git a/Engine/Plugins/NNE/NNERuntimeORT/Source/NNERuntimeORT/NNERuntimeORT.Build.cs b/Engine/Plugins/NNE/NNERuntimeORT/Source/NNERuntimeORT/NNERuntimeORT.Build.cs
--- a/Engine/Plugins/NNE/NNERuntimeORT/Source/NNERuntimeORT/NNERuntimeORT.Build.cs
+++ b/Engine/Plugins/NNE/NNERuntimeORT/Source/NNERuntimeORT/NNERuntimeORT.Build.cs
@@ -12,7 +12,6 @@ public class NNERuntimeORT : ModuleRules
 		{
 			"Core",
 			"CoreUObject",
-			"Engine",
 			"NNE",
 			"NNEOnnxruntime",
 			"Projects",
diff --git a/Engine/Plugins/NNE/NNERuntimeORT/Source/NNERuntimeORT/Private/NNERuntimeORT.cpp b/Engine/Plugins/NNE/NNERuntimeORT/Source/NNERuntimeORT/Private/NNERuntimeORT.cpp
--- a/Engine/Plugins/NNE/NNERuntimeORT/Source/NNERuntimeORT/Private/NNERuntimeORT.cpp
+++ b/Engine/Plugins/NNE/NNERuntimeORT/Source/NNERuntimeORT/Private/NNERuntimeORT.cpp
@@ -2,8 +2,6 @@
 
 #include "NNERuntimeORT.h"
 
-#include "EngineAnalytics.h"
-#include "Kismet/GameplayStatics.h"
 #include "Misc/SecureHash.h"
 #include "NNE.h"
 #include "NNEAttributeMap.h"
@@ -182,16 +180,6 @@ TSharedPtr<UE::NNE::IModelCPU> UNNERuntimeORTCpu::CreateModelCPU(const TObjectPt
 	UE::NNE::IModelCPU* IModel = static_cast<UE::NNE::IModelCPU*>(new UE::NNERuntimeORT::Private::FModelORTCpu(Environment.ToSharedRef(), SharedData));
 	check(IModel != nullptr);
 
-	if (FEngineAnalytics::IsAvailable())
-	{
-		TArray<FAnalyticsEventAttribute> Attributes = MakeAnalyticsEventAttributeArray(
-			TEXT("PlatformName"), UGameplayStatics::GetPlatformName(),
-			TEXT("HashedRuntimeName"), FMD5::HashAnsiString(*GetRuntimeName()),
-			TEXT("ModelDataSize"), SharedData->GetView().Num()
-		);
-		FEngineAnalytics::GetProvider().RecordEvent(TEXT("NeuralNetworkEngine.CreateModel"), Attributes);
-	}
-
 	return TSharedPtr<UE::NNE::IModelCPU>(IModel);
 }
 
@@ -281,16 +269,6 @@ TSharedPtr<UE::NNE::IModelGPU> UNNERuntimeORTDml::CreateModelGPU(const TObjectPt
 	const TSharedPtr<UE::NNE::FSharedModelData> SharedData = ModelData->GetModelData(GetRuntimeName());
 	check(SharedData.IsValid());
 
-	if (FEngineAnalytics::IsAvailable())
-	{
-		TArray<FAnalyticsEventAttribute> Attributes = MakeAnalyticsEventAttributeArray(
-			TEXT("PlatformName"), UGameplayStatics::GetPlatformName(),
-			TEXT("HashedRuntimeName"), FMD5::HashAnsiString(*GetRuntimeName()),
-			TEXT("ModelDataSize"), SharedData->GetView().Num()
-		);
-		FEngineAnalytics::GetProvider().RecordEvent(TEXT("NeuralNetworkEngine.CreateModel"), Attributes);
-	}
-
 	return MakeShared<UE::NNERuntimeORT::Private::FModelORTDmlGPU>(Environment.ToSharedRef(), SharedData);
 #else // PLATFORM_WINDOWS
 	return TSharedPtr<UE::NNE::IModelGPU>();
@@ -315,16 +293,6 @@ TSharedPtr<UE::NNE::IModelRDG> UNNERuntimeORTDml::CreateModelRDG(TObjectPtr<UNNE
 
 	const TSharedRef<UE::NNE::FSharedModelData> SharedData = ModelData->GetModelData(GetRuntimeName()).ToSharedRef();
 
-	if (FEngineAnalytics::IsAvailable())
-	{
-		TArray<FAnalyticsEventAttribute> Attributes = MakeAnalyticsEventAttributeArray(
-			TEXT("PlatformName"), UGameplayStatics::GetPlatformName(),
-			TEXT("HashedRuntimeName"), FMD5::HashAnsiString(*GetRuntimeName()),
-			TEXT("ModelDataSize"), SharedData->GetView().Num()
-		);
-		FEngineAnalytics::GetProvider().RecordEvent(TEXT("NeuralNetworkEngine.CreateModel"), Attributes);
-	}
-
 	return MakeShared<UE::NNERuntimeORT::Private::FModelORTDmlRDG>(Environment.ToSharedRef(), SharedData);
 #else // PLATFORM_WINDOWS
 	return {};
diff --git a/Engine/Source/Editor/NNEEditor/Private/NNEEditorModelDataFactory.cpp b/Engine/Source/Editor/NNEEditor/Private/NNEEditorModelDataFactory.cpp
--- a/Engine/Source/Editor/NNEEditor/Private/NNEEditorModelDataFactory.cpp
+++ b/Engine/Source/Editor/NNEEditor/Private/NNEEditorModelDataFactory.cpp
@@ -3,8 +3,6 @@
 #include "NNEEditorModelDataFactory.h"
 
 #include "Editor.h"
-#include "EngineAnalytics.h"
-#include "Kismet/GameplayStatics.h"
 #include "NNEEditorOnnxFileLoaderHelper.h"
 #include "NNEModelData.h"
 #include "Misc/Paths.h"
@@ -40,16 +38,6 @@ UObject* UNNEModelDataFactory::FactoryCreateFile(UClass* InClass, UObject* InPar
 
 	GEditor->GetEditorSubsystem<UImportSubsystem>()->BroadcastAssetPostImport(this, ModelData);
 
-	if (FEngineAnalytics::IsAvailable())
-	{
-		TArray<FAnalyticsEventAttribute> Attributes = MakeAnalyticsEventAttributeArray(
-			TEXT("PlatformName"), UGameplayStatics::GetPlatformName(),
-			TEXT("FactoryName"), TEXT("UNNEModelDataFactory"),
-			TEXT("ModelFileSize"), ModelFileSize
-		);
-		FEngineAnalytics::GetProvider().RecordEvent(TEXT("NeuralNetworkEngine.FactoryCreateBinary"), Attributes);
-	}
-
 	return ModelData;
 }
 
