[NNE] NNERuntimeORT & NNEModelData log cleanup

Problem:
* Log are inconsistently prefixed by class/function. Also ORT runtime didn't had it's own category.

Solution:
* Remove prefixes from logs, add a log category for the ORT runtime.

Tests:
* Win64 Editor debug : smoke tests

#rnx

diff --git a/Engine/Plugins/NNE/NNERuntimeORT/Source/NNERuntimeORT/Private/NNERuntimeORT.cpp b/Engine/Plugins/NNE/NNERuntimeORT/Source/NNERuntimeORT/Private/NNERuntimeORT.cpp
--- a/Engine/Plugins/NNE/NNERuntimeORT/Source/NNERuntimeORT/Private/NNERuntimeORT.cpp
+++ b/Engine/Plugins/NNE/NNERuntimeORT/Source/NNERuntimeORT/Private/NNERuntimeORT.cpp
@@ -17,6 +17,8 @@
 
 #include UE_INLINE_GENERATED_CPP_BY_NAME(NNERuntimeORT)
 
+DEFINE_LOG_CATEGORY(LogNNERuntimeORT);
+
 FGuid UNNERuntimeORTCpu::GUID = FGuid((int32)'O', (int32)'C', (int32)'P', (int32)'U');
 int32 UNNERuntimeORTCpu::Version = 0x00000004;
 
@@ -87,7 +89,7 @@ TSharedPtr<UE::NNE::FSharedModelData> UNNERuntimeORTCpu::CreateModelData(const F
 
 	if (CanCreateModelData(FileType, FileData, AdditionalFileData, FileId, TargetPlatform) != ECanCreateModelDataStatus::Ok)
 	{
-		UE_LOG(LogNNE, Warning, TEXT("UNNERuntimeORTCpu cannot create the model data with id %s (Filetype: %s)"), *FileId.ToString(EGuidFormats::Digits).ToLower(), *FileType);
+		UE_LOG(LogNNERuntimeORT, Warning, TEXT("Cannot create the CPU model data with id %s (Filetype: %s)"), *FileId.ToString(EGuidFormats::Digits).ToLower(), *FileType);
 		return {};
 	}
 
@@ -170,7 +172,7 @@ TSharedPtr<UE::NNE::IModelCPU> UNNERuntimeORTCpu::CreateModelCPU(const TObjectPt
 
 	if (CanCreateModelCPU(ModelData) != ECanCreateModelCPUStatus::Ok)
 	{
-		UE_LOG(LogNNE, Warning, TEXT("UNNERuntimeORTCpu cannot create a model from the model data with id %s"), *ModelData->GetFileId().ToString(EGuidFormats::Digits));
+		UE_LOG(LogNNERuntimeORT, Warning, TEXT("Cannot create a CPU model from the model data with id %s"), *ModelData->GetFileId().ToString(EGuidFormats::Digits));
 		return TSharedPtr<UE::NNE::IModelCPU>();
 	}
 
@@ -209,7 +211,7 @@ TSharedPtr<UE::NNE::FSharedModelData> UNNERuntimeORTDml::CreateModelData(const F
 
 	if (CanCreateModelData(FileType, FileData, AdditionalFileData, FileId, TargetPlatform) != ECanCreateModelDataStatus::Ok)
 	{
-		UE_LOG(LogNNE, Warning, TEXT("NNERuntimeORTDmlImpl cannot create the model data with id %s (Filetype: %s)"), *FileId.ToString(EGuidFormats::Digits).ToLower(), *FileType);
+		UE_LOG(LogNNERuntimeORT, Warning, TEXT("Cannot create the Dml model data with id %s (Filetype: %s)"), *FileId.ToString(EGuidFormats::Digits).ToLower(), *FileType);
 		return {};
 	}
 
@@ -262,7 +264,7 @@ TSharedPtr<UE::NNE::IModelGPU> UNNERuntimeORTDml::CreateModelGPU(const TObjectPt
 
 	if (CanCreateModelGPU(ModelData) != ECanCreateModelGPUStatus::Ok)
 	{
-		UE_LOG(LogNNE, Warning, TEXT("UNNERuntimeORTDml cannot create a model GPU from the model data with id %s"), *ModelData->GetFileId().ToString(EGuidFormats::Digits));
+		UE_LOG(LogNNERuntimeORT, Warning, TEXT("Cannot create a GPU model from the model data with id %s"), *ModelData->GetFileId().ToString(EGuidFormats::Digits));
 		return {};
 	}
 
@@ -287,7 +289,7 @@ TSharedPtr<UE::NNE::IModelRDG> UNNERuntimeORTDml::CreateModelRDG(TObjectPtr<UNNE
 
 	if (CanCreateModelRDG(ModelData) != ECanCreateModelRDGStatus::Ok)
 	{
-		UE_LOG(LogNNE, Warning, TEXT("UNNERuntimeORTDml cannot create a model RDG from the model data with id %s"), *ModelData->GetFileId().ToString(EGuidFormats::Digits));
+		UE_LOG(LogNNERuntimeORT, Warning, TEXT("Cannot create a RDG model from the model data with id %s"), *ModelData->GetFileId().ToString(EGuidFormats::Digits));
 		return {};
 	}
 
diff --git a/Engine/Plugins/NNE/NNERuntimeORT/Source/NNERuntimeORT/Private/NNERuntimeORT.h b/Engine/Plugins/NNE/NNERuntimeORT/Source/NNERuntimeORT/Private/NNERuntimeORT.h
--- a/Engine/Plugins/NNE/NNERuntimeORT/Source/NNERuntimeORT/Private/NNERuntimeORT.h
+++ b/Engine/Plugins/NNE/NNERuntimeORT/Source/NNERuntimeORT/Private/NNERuntimeORT.h
@@ -12,6 +12,8 @@
 
 #include "NNERuntimeORT.generated.h"
 
+DECLARE_LOG_CATEGORY_EXTERN(LogNNERuntimeORT, Log, All);
+
 namespace UE::NNERuntimeORT::Private
 {
 	class FEnvironment;
diff --git a/Engine/Plugins/NNE/NNERuntimeORT/Source/NNERuntimeORT/Private/NNERuntimeORTModel.cpp b/Engine/Plugins/NNE/NNERuntimeORT/Source/NNERuntimeORT/Private/NNERuntimeORTModel.cpp
--- a/Engine/Plugins/NNE/NNERuntimeORT/Source/NNERuntimeORT/Private/NNERuntimeORTModel.cpp
+++ b/Engine/Plugins/NNE/NNERuntimeORT/Source/NNERuntimeORT/Private/NNERuntimeORTModel.cpp
@@ -80,7 +80,7 @@ namespace Detail
 
 		if (ModelBuffer.Num() == 0)
 		{
-			UE_LOG(LogNNE, Error, TEXT("NNERuntimeORT::Private::Details::CreateSession(): Input model data is empty."));
+			UE_LOG(LogNNERuntimeORT, Error, TEXT("Cannot create ORT session: Input model data is empty."));
 			return false;
 		}
 
@@ -99,7 +99,7 @@ namespace Detail
 				if (!FFileHelper::SaveArrayToFile(ModelBuffer, *Filepath))
 				{
 					IFileManager::Get().DeleteDirectory(*TempDirForModelWithExternalData, false, true);
-					UE_LOG(LogNNE, Error, TEXT("Large models are an experimental feature at the moment. NNERuntimeORT::Private::Details::CreateSession() could not write model to disk at %s."), *Filepath);
+					UE_LOG(LogNNERuntimeORT, Error, TEXT("Large models are an experimental feature at the moment. Could not write model to disk at %s."), *Filepath);
 					return false;
 				}
 
@@ -111,7 +111,7 @@ namespace Detail
 					if (!FFileHelper::SaveArrayToFile(AdditionalDataBuffer, *AdditionalDataFilename))
 					{
 						IFileManager::Get().DeleteDirectory(*TempDirForModelWithExternalData, false, true);
-						UE_LOG(LogNNE, Error, TEXT("Large models are an experimental feature at the moment. NNERuntimeORT::Private::Details::CreateSession() could not write additional data to disk at %s."), *AdditionalDataFilename);
+						UE_LOG(LogNNERuntimeORT, Error, TEXT("Large models are an experimental feature at the moment. Could not write additional data to disk at %s."), *AdditionalDataFilename);
 						return false;
 					}
 				}
@@ -152,7 +152,7 @@ FModelInstanceORTBase<ModelInterface, TensorBinding>::~FModelInstanceORTBase()
 	{
 		if (!IFileManager::Get().DeleteDirectory(*TempDirForModelWithExternalData, false, true))
 		{
-			UE_LOG(LogNNE, Warning, TEXT("Large models are an experimental feature at the moment. FModelInstanceORTBase could not delete temp directy %s on model instance destruction."), *TempDirForModelWithExternalData);
+			UE_LOG(LogNNERuntimeORT, Warning, TEXT("Large models are an experimental feature at the moment. Could not delete temp directy %s on model instance destruction."), *TempDirForModelWithExternalData);
 		}
 	}
 }
@@ -166,36 +166,36 @@ bool FModelInstanceORTBase<ModelInterface, TensorBinding>::Init(TConstArrayView6
 	{
 		if (!InitializedAndConfigureMembers())
 		{
-			UE_LOG(LogNNE, Error, TEXT("FModelInstanceORTBase::Init(): InitializedAndConfigureMembers failed."));
+			UE_LOG(LogNNERuntimeORT, Error, TEXT("InitializedAndConfigureMembers failed."));
 			return false;
 		}
 
 		if (!Detail::CreateSession(ModelData, *SessionOptions, *Environment, Session, TempDirForModelWithExternalData))
 		{
-			UE_LOG(LogNNE, Error, TEXT("FModelInstanceORTBase::Init(): Session creation failed."));
+			UE_LOG(LogNNERuntimeORT, Error, TEXT("Session creation failed."));
 			return false;
 		}
 
 		if (!ConfigureTensors(true))
 		{
-			UE_LOG(LogNNE, Error, TEXT("FModelInstanceORTBase::Init(): Failed to configure Inputs tensors."));
+			UE_LOG(LogNNERuntimeORT, Error, TEXT("Failed to configure Inputs tensors."));
 			return false;
 		}
 		if (!ConfigureTensors(false))
 		{
-			UE_LOG(LogNNE, Error, TEXT("FModelInstanceORTBase::Init(): Failed to configure Outputs tensors."));
+			UE_LOG(LogNNERuntimeORT, Error, TEXT("Failed to configure Outputs tensors."));
 			return false;
 		}
 	}
 #if WITH_EDITOR
 	catch (const Ort::Exception& Exception)
 	{
-		UE_LOG(LogNNE, Error, TEXT("%s"), UTF8_TO_TCHAR(Exception.what()));
+		UE_LOG(LogNNERuntimeORT, Error, TEXT("%s"), UTF8_TO_TCHAR(Exception.what()));
 		return false;
 	}
 	catch (...)
 	{
-		UE_LOG(LogNNE, Error, TEXT("Unknown exception!"));
+		UE_LOG(LogNNERuntimeORT, Error, TEXT("Unknown exception!"));
 		return false;
 	}
 #endif // WITH_EDITOR
@@ -319,14 +319,14 @@ typename ModelInterface::ERunSyncStatus FModelInstanceORTBase<ModelInterface, Te
 
 	if (!Session.IsValid())
 	{
-		UE_LOG(LogNNE, Error, TEXT("FModelInstanceORTBase::RunSync(): Called without a Session, FModelInstanceORT::Init() should have been called."));
+		UE_LOG(LogNNERuntimeORT, Error, TEXT("Called without a Session."));
 		return ModelInterface::ERunSyncStatus::Fail;
 	}
 
 	// Verify the model inputs were prepared
 	if (NNE::Internal::FModelInstanceBase<ModelInterface>::InputTensorShapes.IsEmpty())
 	{
-		UE_LOG(LogNNE, Error, TEXT("FModelInstanceORTBase::RunSync(): Input shapes are not set, please call SetInputTensorShapes."));
+		UE_LOG(LogNNERuntimeORT, Error, TEXT("Input shapes are not set, please call SetInputTensorShapes."));
 		return ModelInterface::ERunSyncStatus::Fail;
 	}
 
@@ -336,7 +336,7 @@ typename ModelInterface::ERunSyncStatus FModelInstanceORTBase<ModelInterface, Te
 
 	if (InInputBindings.Num() != InputTensors.Num())
 	{
-		UE_LOG(LogNNE, Error, TEXT("FModelInstanceORTBase::RunSync(): Input bindings need to match input tensor descriptor count (got %d, expected %d)."), InInputBindings.Num(), InputTensors.Num());
+		UE_LOG(LogNNERuntimeORT, Error, TEXT("Input bindings need to match input tensor descriptor count (got %d, expected %d)."), InInputBindings.Num(), InputTensors.Num());
 		return ModelInterface::ERunSyncStatus::Fail;
 	}
 
@@ -344,7 +344,7 @@ typename ModelInterface::ERunSyncStatus FModelInstanceORTBase<ModelInterface, Te
 
 	if (!InOutputBindings.IsEmpty() && InOutputBindings.Num() != OutputTensorNames.Num())
 	{
-		UE_LOG(LogNNE, Error, TEXT("FModelInstanceORTBase::RunSync(): Output binding can be empty or needs to match output tensor descriptor count (got %d, expected %d)."), InOutputBindings.Num(), OutputTensorNames.Num());
+		UE_LOG(LogNNERuntimeORT, Error, TEXT("Output binding can be empty or needs to match output tensor descriptor count (got %d, expected %d)."), InOutputBindings.Num(), OutputTensorNames.Num());
 		return ModelInterface::ERunSyncStatus::Fail;
 	}
 
@@ -360,13 +360,13 @@ typename ModelInterface::ERunSyncStatus FModelInstanceORTBase<ModelInterface, Te
 
 			if (!Binding.Data && Binding.SizeInBytes != 0)
 			{
-				UE_LOG(LogNNE, Error, TEXT("FModelInstanceORTBase::RunSync(): Binding input tensor %d is not set but given size is non-zero %d."), i, Binding.SizeInBytes);
+				UE_LOG(LogNNERuntimeORT, Error, TEXT("Binding input tensor %d is not set but given size is non-zero %d."), i, Binding.SizeInBytes);
 				return ModelInterface::ERunSyncStatus::Fail;
 			}
 
 			if (Binding.SizeInBytes != Tensor.GetDataSize())
 			{
-				UE_LOG(LogNNE, Error, TEXT("FModelInstanceORTBase::RunSync(): Binding input tensor %d size does not match size given by tensor descriptor (got %d, expected %d)."), i, Binding.SizeInBytes, Tensor.GetDataSize());
+				UE_LOG(LogNNERuntimeORT, Error, TEXT("Binding input tensor %d size does not match size given by tensor descriptor (got %d, expected %d)."), i, Binding.SizeInBytes, Tensor.GetDataSize());
 				return ModelInterface::ERunSyncStatus::Fail;
 			}
 
@@ -420,12 +420,12 @@ typename ModelInterface::ERunSyncStatus FModelInstanceORTBase<ModelInterface, Te
 #if WITH_EDITOR
 	catch (const Ort::Exception& Exception)
 	{
-		UE_LOG(LogNNE, Error, TEXT("%s"), UTF8_TO_TCHAR(Exception.what()));
+		UE_LOG(LogNNERuntimeORT, Error, TEXT("%s"), UTF8_TO_TCHAR(Exception.what()));
 		return ModelInterface::ERunSyncStatus::Fail;
 	}
 	catch (...)
 	{
-		UE_LOG(LogNNE, Error, TEXT("Unknown exception!"));
+		UE_LOG(LogNNERuntimeORT, Error, TEXT("Unknown exception!"));
 		return ModelInterface::ERunSyncStatus::Fail;
 	}
 #endif // WITH_EDITOR
@@ -544,7 +544,7 @@ FModelInstanceORTDmlRDG::~FModelInstanceORTDmlRDG()
 	{
 		if (!IFileManager::Get().DeleteDirectory(*TempDirForModelWithExternalData, false, true))
 		{
-			UE_LOG(LogNNE, Warning, TEXT("Large models are an experimental feature at the moment. FModelInstanceORTDmlRDG could not delete temp directy %s on model instance destruction."), *TempDirForModelWithExternalData);
+			UE_LOG(LogNNERuntimeORT, Warning, TEXT("Large models are an experimental feature at the moment. FModelInstanceORTDmlRDG could not delete temp directy %s on model instance destruction."), *TempDirForModelWithExternalData);
 		}
 	}
 }
@@ -560,7 +560,7 @@ bool FModelInstanceORTDmlRDG::Init()
 		SessionOptions = CreateSessionOptionsForDirectML(Environment);
 		if (!SessionOptions.IsValid())
 		{
-			UE_LOG(LogNNE, Error, TEXT("FModelInstanceORTDmlRDG::Init(): Failed to configure session options for DirectML Execution Provider."));
+			UE_LOG(LogNNERuntimeORT, Error, TEXT("Failed to configure session options for DirectML Execution Provider."));
 			return false;
 		}
 		
@@ -568,25 +568,25 @@ bool FModelInstanceORTDmlRDG::Init()
 
 		if (!Detail::CreateSession(ModelData->GetView(), *SessionOptions, *Environment, Session, TempDirForModelWithExternalData))
 		{
-			UE_LOG(LogNNE, Error, TEXT("FModelInstanceORTDmlRDG::Init(): Session creation failed."));
+			UE_LOG(LogNNERuntimeORT, Error, TEXT("Session creation failed."));
 			return false;
 		}
 
 		if (!ConfigureTensors(*Session))
 		{
-			UE_LOG(LogNNE, Error, TEXT("FModelInstanceORTDmlRDG::Init(): Failed to configure Inputs tensors."));
+			UE_LOG(LogNNERuntimeORT, Error, TEXT("Failed to configure Inputs tensors."));
 			return false;
 		}
 	}
 #if WITH_EDITOR
 	catch (const Ort::Exception& Exception)
 	{
-		UE_LOG(LogNNE, Error, TEXT("%s"), UTF8_TO_TCHAR(Exception.what()));
+		UE_LOG(LogNNERuntimeORT, Error, TEXT("%s"), UTF8_TO_TCHAR(Exception.what()));
 		return false;
 	}
 	catch (...)
 	{
-		UE_LOG(LogNNE, Error, TEXT("Unknown exception!"));
+		UE_LOG(LogNNERuntimeORT, Error, TEXT("Unknown exception!"));
 		return false;
 	}
 #endif // WITH_EDITOR
@@ -598,12 +598,12 @@ bool FModelInstanceORTDmlRDG::ConfigureTensors(const Ort::Session& ActiveSession
 {
 	if (!ConfigureTensors(ActiveSession, true))
 	{
-		UE_LOG(LogNNE, Error, TEXT("FModelInstanceORTDmlRDG::ConfigureTensors(): Failed to configure Inputs tensors."));
+		UE_LOG(LogNNERuntimeORT, Error, TEXT("Failed to configure Inputs tensors."));
 		return false;
 	}
 	if (!ConfigureTensors(ActiveSession, false))
 	{
-		UE_LOG(LogNNE, Error, TEXT("FModelInstanceORTDmlRDG::ConfigureTensors(): Failed to configure Outputs tensors."));
+		UE_LOG(LogNNERuntimeORT, Error, TEXT("Failed to configure Outputs tensors."));
 		return false;
 	}
 
@@ -719,7 +719,7 @@ FModelInstanceORTDmlRDG::ESetInputTensorShapesStatus FModelInstanceORTDmlRDG::Se
 			}
 			else
 			{
-				UE_LOG(LogNNE, Warning, TEXT("One or more output tensors contain free dimensions, but input tensors are all concrete!"));
+				UE_LOG(LogNNERuntimeORT, Warning, TEXT("One or more output tensors contain free dimensions, but input tensors are all concrete!"));
 				return ESetInputTensorShapesStatus::Fail;
 			}
 		}
@@ -731,7 +731,7 @@ FModelInstanceORTDmlRDG::ESetInputTensorShapesStatus FModelInstanceORTDmlRDG::Se
 	SessionOptions = CreateSessionOptionsForDirectML(Environment);
 	if (!SessionOptions.IsValid())
 	{
-		UE_LOG(LogNNE, Error, TEXT("Failed to recreate session options!"));
+		UE_LOG(LogNNERuntimeORT, Error, TEXT("Failed to recreate session options!"));
 		return ESetInputTensorShapesStatus::Fail;
 	}
 
@@ -769,19 +769,19 @@ FModelInstanceORTDmlRDG::ESetInputTensorShapesStatus FModelInstanceORTDmlRDG::Se
 	{
 		if (!Detail::CreateSession(ModelData->GetView(), *SessionOptions, *Environment, Session, TempDirForModelWithExternalData))
 		{
-			UE_LOG(LogNNE, Error, TEXT("Failed to recreate session!"));
+			UE_LOG(LogNNERuntimeORT, Error, TEXT("Failed to recreate session!"));
 			return ESetInputTensorShapesStatus::Fail;
 		}
 	}
 #if WITH_EDITOR
 	catch (const Ort::Exception& Exception)
 	{
-		UE_LOG(LogNNE, Error, TEXT("%s"), UTF8_TO_TCHAR(Exception.what()));
+		UE_LOG(LogNNERuntimeORT, Error, TEXT("%s"), UTF8_TO_TCHAR(Exception.what()));
 		return ESetInputTensorShapesStatus::Fail;
 	}
 	catch (...)
 	{
-		UE_LOG(LogNNE, Error, TEXT("Unknown exception!"));
+		UE_LOG(LogNNERuntimeORT, Error, TEXT("Unknown exception!"));
 		return ESetInputTensorShapesStatus::Fail;
 	}
 #endif // WITH_EDITOR
@@ -789,7 +789,7 @@ FModelInstanceORTDmlRDG::ESetInputTensorShapesStatus FModelInstanceORTDmlRDG::Se
 	// Need to configure output tensors with new session (to apply free dimension overrides)
 	if (!ConfigureTensors(*Session, false))
 	{
-		UE_LOG(LogNNE, Error, TEXT("Failed to configure tensors!"));
+		UE_LOG(LogNNERuntimeORT, Error, TEXT("Failed to configure tensors!"));
 		return ESetInputTensorShapesStatus::Fail;
 	}
 
@@ -810,11 +810,11 @@ FModelInstanceORTDmlRDG::ESetInputTensorShapesStatus FModelInstanceORTDmlRDG::Se
 			{
 				if (SymbolicTensorDesc.GetShape().GetData()[j] < 0)
 				{
-					UE_LOG(LogNNE, Warning, TEXT("Tensor '%hs' has free dimension '%s'."), OutputTensorNames[i], *OutputSymbolicDimensionNames[i][j]);
+					UE_LOG(LogNNERuntimeORT, Warning, TEXT("Tensor '%hs' has free dimension '%s'."), OutputTensorNames[i], *OutputSymbolicDimensionNames[i][j]);
 				}
 			}
 
-			UE_LOG(LogNNE, Error, TEXT("One or more output tensors contain free dimensions!"));
+			UE_LOG(LogNNERuntimeORT, Error, TEXT("One or more output tensors contain free dimensions!"));
 			return ESetInputTensorShapesStatus::Fail;
 		}
 	}
@@ -863,13 +863,13 @@ FModelInstanceORTDmlRDG::EEnqueueRDGStatus FModelInstanceORTDmlRDG::EnqueueRDG(F
 
 	if (!Session.IsValid())
 	{
-		UE_LOG(LogNNE, Error, TEXT("FModelInstanceORTDmlRDG::EnqueueRDG(): Called without a Session, FModelInstanceORTDmlRDG::Init() should have been called."));
+		UE_LOG(LogNNERuntimeORT, Error, TEXT("Invalid Session, Init() should have been called."));
 		return EEnqueueRDGStatus::Fail;
 	}
 
 	if (InputTensorShapes.IsEmpty())
 	{
-		UE_LOG(LogNNE, Error, TEXT("FModelInstanceORTDmlRDG::EnqueueRDG(): Input shapes are not set, please call SetInputTensorShapes."));
+		UE_LOG(LogNNERuntimeORT, Error, TEXT("Input shapes are not set, please call SetInputTensorShapes."));
 		return EEnqueueRDGStatus::Fail;
 	}
 
@@ -879,7 +879,7 @@ FModelInstanceORTDmlRDG::EEnqueueRDGStatus FModelInstanceORTDmlRDG::EnqueueRDG(F
 
 	if (Inputs.Num() != InputTensors.Num())
 	{
-		UE_LOG(LogNNE, Error, TEXT("FModelInstanceORTDmlRDG::EnqueueRDG(): Input bindings need to match input tensor descriptor count (got %d, expected %d)."), Inputs.Num(), InputTensors.Num());
+		UE_LOG(LogNNERuntimeORT, Error, TEXT("Input bindings need to match input tensor descriptor count (got %d, expected %d)."), Inputs.Num(), InputTensors.Num());
 		return EEnqueueRDGStatus::Fail;
 	}
 
@@ -887,7 +887,7 @@ FModelInstanceORTDmlRDG::EEnqueueRDGStatus FModelInstanceORTDmlRDG::EnqueueRDG(F
 
 	if (!Outputs.IsEmpty() && Outputs.Num() != OutputTensorNames.Num())
 	{
-		UE_LOG(LogNNE, Error, TEXT("FModelInstanceORTDmlRDG::EnqueueRDG(): Output binding can be empty or needs to match output tensor descriptor count (got %d, expected %d)."), Outputs.Num(), OutputTensorNames.Num());
+		UE_LOG(LogNNERuntimeORT, Error, TEXT("Output binding can be empty or needs to match output tensor descriptor count (got %d, expected %d)."), Outputs.Num(), OutputTensorNames.Num());
 		return EEnqueueRDGStatus::Fail;
 	}
 
@@ -897,13 +897,13 @@ FModelInstanceORTDmlRDG::EEnqueueRDGStatus FModelInstanceORTDmlRDG::EnqueueRDG(F
 		const NNE::FTensorBindingRDG& Binding = Inputs[i];
 		if (!Binding.Buffer && InputTensors[i].GetDataSize() != 0)
 		{
-			UE_LOG(LogNNE, Error, TEXT("FModelInstanceORTDmlRDG::EnqueueRDG(): Binding input tensor %d is not set but given size by tensor descriptor is non-zero %d."), i, InputTensors[i].GetDataSize());
+			UE_LOG(LogNNERuntimeORT, Error, TEXT("Binding input tensor %d is not set but given size by tensor descriptor is non-zero %d."), i, InputTensors[i].GetDataSize());
 			return EEnqueueRDGStatus::Fail;
 		}
 
 		if (Binding.Buffer->Desc.GetSize() != InputTensors[i].GetDataSize())
 		{
-			UE_LOG(LogNNE, Error, TEXT("FModelInstanceORTDmlRDG::EnqueueRDG(): Binding input tensor %d size does not match size given by tensor descriptor (got %d, expected %d)."), i, Binding.Buffer->Desc.GetSize(), InputTensors[i].GetDataSize());
+			UE_LOG(LogNNERuntimeORT, Error, TEXT("Binding input tensor %d size does not match size given by tensor descriptor (got %d, expected %d)."), i, Binding.Buffer->Desc.GetSize(), InputTensors[i].GetDataSize());
 			return EEnqueueRDGStatus::Fail;
 		}
 
@@ -989,11 +989,11 @@ FModelInstanceORTDmlRDG::EEnqueueRDGStatus FModelInstanceORTDmlRDG::EnqueueRDG(F
 #if WITH_EDITOR
 				catch (const Ort::Exception& Exception)
 				{
-					UE_LOG(LogNNE, Error, TEXT("FModelInstanceORTDmlRDG::EnqueueRDG(): %s"), UTF8_TO_TCHAR(Exception.what()));
+					UE_LOG(LogNNERuntimeORT, Error, TEXT("ORT Exception: %s"), UTF8_TO_TCHAR(Exception.what()));
 				}
 				catch (...)
 				{
-					UE_LOG(LogNNE, Error, TEXT("FModelInstanceORTDmlRDG::EnqueueRDG(): Unknown exception!"));
+					UE_LOG(LogNNERuntimeORT, Error, TEXT("ORT Exception: Unknown!"));
 				}
 #endif // WITH_EDITOR
 			}, false);
diff --git a/Engine/Plugins/NNE/NNERuntimeORT/Source/NNERuntimeORT/Private/NNERuntimeORTModule.cpp b/Engine/Plugins/NNE/NNERuntimeORT/Source/NNERuntimeORT/Private/NNERuntimeORTModule.cpp
--- a/Engine/Plugins/NNE/NNERuntimeORT/Source/NNERuntimeORT/Private/NNERuntimeORTModule.cpp
+++ b/Engine/Plugins/NNE/NNERuntimeORT/Source/NNERuntimeORT/Private/NNERuntimeORTModule.cpp
@@ -24,7 +24,7 @@ namespace UE::NNERuntimeORT::Private
 
 			if (!FPaths::FileExists(DllPath))
 			{
-				UE_LOG(LogNNE, Error, TEXT("Failed to find the third party library %s."), *DllPath);
+				UE_LOG(LogNNERuntimeORT, Error, TEXT("Failed to find the third party library %s."), *DllPath);
 				return false;
 			}
 			
@@ -32,7 +32,7 @@ namespace UE::NNERuntimeORT::Private
 
 			if (!DllHandle)
 			{
-				UE_LOG(LogNNE, Error, TEXT("Failed to load the third party library %s."), *DllPath);
+				UE_LOG(LogNNERuntimeORT, Error, TEXT("Failed to load the third party library %s."), *DllPath);
 				return false;
 			}
 
@@ -70,7 +70,7 @@ void FNNERuntimeORTModule::StartupModule()
 
 	if (!DllHelper::GetDllHandle(OrtSharedLibPath, DllHandles))
 	{
-		UE_LOG(LogNNE, Error, TEXT("Failed to load ONNX Runtime shared library. ORT Runtimes won't be available."));
+		UE_LOG(LogNNERuntimeORT, Error, TEXT("Failed to load ONNX Runtime shared library. ORT Runtimes won't be available."));
 		return;
 	}
 
@@ -85,14 +85,14 @@ void FNNERuntimeORTModule::StartupModule()
 	bDirectMLDllLoaded = DllHelper::GetDllHandle(DirectMLSharedLibPath, DllHandles);
 	if (!bDirectMLDllLoaded)
 	{
-		UE_LOG(LogNNE, Error, TEXT("Failed to load DirectML shared library. ORT Dml Runtime won't be available."));
+		UE_LOG(LogNNERuntimeORT, Error, TEXT("Failed to load DirectML shared library. ORT Dml Runtime won't be available."));
 	}
 #endif // PLATFORM_WINDOWS
 
 	TUniquePtr<UE::NNEOnnxruntime::OrtApiFunctions> OrtApiFunctions = UE::NNEOnnxruntime::LoadApiFunctions(OrtDllHandle);
 	if (!OrtApiFunctions.IsValid())
 	{
-		UE_LOG(LogNNE, Fatal, TEXT("Failed to load ONNX Runtime shared library functions!"));
+		UE_LOG(LogNNERuntimeORT, Fatal, TEXT("Failed to load ONNX Runtime shared library functions!"));
 		return;
 	}
 
@@ -171,8 +171,8 @@ void FNNERuntimeORTModule::ShutdownModule()
 #if WITH_EDITOR
 void FNNERuntimeORTModule::OnSettingsChanged(UObject* InObject, struct FPropertyChangedEvent& InPropertyChangedEvent)
 {
-	UE_LOG(LogNNE, Log, TEXT("Settings %s changed: %s"), *InObject->GetName(), *InPropertyChangedEvent.GetPropertyName().ToString());
-	UE_LOG(LogNNE, Warning, TEXT("It is recommended to restart the Editor if settings %s changed! Otherwise they might not be fully applied."), *InObject->GetName());
+	UE_LOG(LogNNERuntimeORT, Log, TEXT("Settings %s changed: %s"), *InObject->GetName(), *InPropertyChangedEvent.GetPropertyName().ToString());
+	UE_LOG(LogNNERuntimeORT, Warning, TEXT("It is recommended to restart the Editor if settings %s changed! Otherwise they might not be fully applied."), *InObject->GetName());
 
 	UE::NNERuntimeORT::Private::EnvironmentHelper::CreateOrtEnvFromSettings(CastChecked<UNNERuntimeORTSettings>(InObject), *Environment);
 }
diff --git a/Engine/Plugins/NNE/NNERuntimeORT/Source/NNERuntimeORT/Private/NNERuntimeORTUtils.cpp b/Engine/Plugins/NNE/NNERuntimeORT/Source/NNERuntimeORT/Private/NNERuntimeORTUtils.cpp
--- a/Engine/Plugins/NNE/NNERuntimeORT/Source/NNERuntimeORT/Private/NNERuntimeORTUtils.cpp
+++ b/Engine/Plugins/NNE/NNERuntimeORT/Source/NNERuntimeORT/Private/NNERuntimeORTUtils.cpp
@@ -6,7 +6,7 @@
 #include "HAL/PlatformTime.h"
 #include "Misc/FileHelper.h"
 #include "Misc/Paths.h"
-#include "NNE.h"
+#include "NNERuntimeORT.h"
 #include "NNERuntimeORTEnv.h"
 
 #if PLATFORM_WINDOWS
@@ -214,14 +214,14 @@ TUniquePtr<Ort::SessionOptions> CreateSessionOptionsForDirectML(const TSharedRef
 		Ort::ThrowOnError(Ort::GetApi().GetExecutionProviderApi("DML", ORT_API_VERSION, reinterpret_cast<const void**>(&DmlApi)));
 		if (!DmlApi)
 		{
-			UE_LOG(LogNNE, Error, TEXT("Ort DirectML Api not available!"));
+			UE_LOG(LogNNERuntimeORT, Error, TEXT("Ort DirectML Api not available!"));
 			return {};
 		}
 
 		OrtStatusPtr Status = DmlApi->SessionOptionsAppendExecutionProvider_DML(*SessionOptions.Get(), DeviceIndex);
 		if (Status)
 		{
-			UE_LOG(LogNNE, Error, TEXT("Failed to add DirectML execution provider to OnnxRuntime session options: %s"), ANSI_TO_TCHAR(Ort::GetApi().GetErrorMessage(Status)));
+			UE_LOG(LogNNERuntimeORT, Error, TEXT("Failed to add DirectML execution provider to OnnxRuntime session options: %s"), ANSI_TO_TCHAR(Ort::GetApi().GetErrorMessage(Status)));
 			return {};
 		}
 
@@ -230,7 +230,7 @@ TUniquePtr<Ort::SessionOptions> CreateSessionOptionsForDirectML(const TSharedRef
 
 	if (!GDynamicRHI)
 	{
-		UE_LOG(LogNNE, Error, TEXT("Error:No RHI found, could not initialize"));
+		UE_LOG(LogNNERuntimeORT, Error, TEXT("No RHI found, could not initialize"));
 		return {};
 	}
 
@@ -244,12 +244,12 @@ TUniquePtr<Ort::SessionOptions> CreateSessionOptionsForDirectML(const TSharedRef
 	{
 		if (GDynamicRHI)
 		{
-			UE_LOG(LogNNE, Error, TEXT("Error:%s RHI is not supported by DirectML, please use D3D12."), GDynamicRHI->GetName());
+			UE_LOG(LogNNERuntimeORT, Error, TEXT("%s RHI is not supported by DirectML, please use D3D12."), GDynamicRHI->GetName());
 			return {};
 		}
 		else
 		{
-			UE_LOG(LogNNE, Error, TEXT("Error:No RHI found"));
+			UE_LOG(LogNNERuntimeORT, Error, TEXT("No RHI found"));
 			return {};
 		}
 	}
@@ -261,7 +261,7 @@ TUniquePtr<Ort::SessionOptions> CreateSessionOptionsForDirectML(const TSharedRef
 
 	if (!D3D12Device)
 	{
-		UE_LOG(LogNNE, Error, TEXT("Failed to get D3D12 Device from RHI for device index %d"), DeviceIndex);
+		UE_LOG(LogNNERuntimeORT, Error, TEXT("Failed to get D3D12 Device from RHI for device index %d"), DeviceIndex);
 		return {};
 	}
 
@@ -278,7 +278,7 @@ TUniquePtr<Ort::SessionOptions> CreateSessionOptionsForDirectML(const TSharedRef
 
 	if (FAILED(Res) || !DmlDevice)
 	{
-		UE_LOG(LogNNE, Error, TEXT("Failed to create DirectML device, DMLCreateDevice error code :%x"), Res);
+		UE_LOG(LogNNERuntimeORT, Error, TEXT("Failed to create DirectML device, DMLCreateDevice error code :%x"), Res);
 		return {};
 	}
 
@@ -289,7 +289,7 @@ TUniquePtr<Ort::SessionOptions> CreateSessionOptionsForDirectML(const TSharedRef
 
 	if (!DmlApi)
 	{
-		UE_LOG(LogNNE, Error, TEXT("Ort DirectML Api not available!"));
+		UE_LOG(LogNNERuntimeORT, Error, TEXT("Ort DirectML Api not available!"));
 		return {};
 	}
 
@@ -297,7 +297,7 @@ TUniquePtr<Ort::SessionOptions> CreateSessionOptionsForDirectML(const TSharedRef
 
 	if (Status)
 	{
-		UE_LOG(LogNNE, Error, TEXT("Failed to add DirectML execution provider to OnnxRuntime session options: %s"), ANSI_TO_TCHAR(Ort::GetApi().GetErrorMessage(Status)));
+		UE_LOG(LogNNERuntimeORT, Error, TEXT("Failed to add DirectML execution provider to OnnxRuntime session options: %s"), ANSI_TO_TCHAR(Ort::GetApi().GetErrorMessage(Status)));
 		return {};
 	}
 
@@ -332,12 +332,12 @@ bool OptimizeModel(const TSharedRef<FEnvironment> &Environment, Ort::SessionOpti
 #if WITH_EDITOR
 	catch (const Ort::Exception& Exception)
 	{
-		UE_LOG(LogNNE, Error, TEXT("%s"), UTF8_TO_TCHAR(Exception.what()));
+		UE_LOG(LogNNERuntimeORT, Error, TEXT("ORT Exception: %s"), UTF8_TO_TCHAR(Exception.what()));
 		return false;
 	}
 	catch (...)
 	{
-		UE_LOG(LogNNE, Error, TEXT("Unknown exception!"));
+		UE_LOG(LogNNERuntimeORT, Error, TEXT("ORT Exception: Unknown!"));
 		return false;
 	}
 #endif // WITH_EDITOR
diff --git a/Engine/Source/Runtime/NNE/Private/NNEModelData.cpp b/Engine/Source/Runtime/NNE/Private/NNEModelData.cpp
--- a/Engine/Source/Runtime/NNE/Private/NNEModelData.cpp
+++ b/Engine/Source/Runtime/NNE/Private/NNEModelData.cpp
@@ -70,7 +70,7 @@ namespace UE::NNE::ModelData
 		}
 		else
 		{
-			UE_LOG(LogNNE, Error, TEXT("UNNEModelData: No runtime '%s' found. Valid runtimes are: "), *RuntimeName);
+			UE_LOG(LogNNE, Error, TEXT("No runtime '%s' found. Valid runtimes are: "), *RuntimeName);
 			TArray<FString> Runtimes = UE::NNE::GetAllRuntimeNames();
 			for (int32 i = 0; i < Runtimes.Num(); i++)
 			{
@@ -246,7 +246,7 @@ void UNNEModelData::Serialize(FArchive& Ar)
 				Ar << Data_Size32Bits;
 				ModelData.Add(Name, MakeShared<UE::NNE::FSharedModelData>(MakeSharedBufferFromArray(MoveTemp(Data_Size32Bits)), 0));
 			}
-			UE_LOG(LogNNE, Warning, TEXT("[DEPRECATION] UNNEModelData: The asset %s (v0) is deprecated. Please right-click the asset and select 'Save' to update it to the latest version."), *this->GetName());
+			UE_LOG(LogNNE, Warning, TEXT("[DEPRECATION] The asset %s (v0) is deprecated. Please right-click the asset and select 'Save' to update it to the latest version."), *this->GetName());
 			break;
 
 		case UE::NNE::ModelData::Version::V1:
@@ -267,7 +267,7 @@ void UNNEModelData::Serialize(FArchive& Ar)
 				Ar << Data_Size32Bits;
 				ModelData.Add(Name, MakeShared<UE::NNE::FSharedModelData>(MakeSharedBufferFromArray(MoveTemp(Data_Size32Bits)), 0));
 			}
-			UE_LOG(LogNNE, Warning, TEXT("[DEPRECATION] UNNEModelData: The asset %s (v1) is deprecated. Please right-click the asset and select 'Save' to update it to the latest version."), *this->GetName());
+			UE_LOG(LogNNE, Warning, TEXT("[DEPRECATION] The asset %s (v1) is deprecated. Please right-click the asset and select 'Save' to update it to the latest version."), *this->GetName());
 			break;
 
 		case UE::NNE::ModelData::Version::V2:
@@ -341,7 +341,7 @@ void UNNEModelData::Serialize(FArchive& Ar)
 			break;
 
 		default:
-			UE_LOG(LogNNE, Error, TEXT("UNNEModelData: Unknown asset version %d: Deserialisation failed, please reimport the original model."), Ar.CustomVer(UE::NNE::ModelData::GUID));
+			UE_LOG(LogNNE, Error, TEXT("Unknown asset version %d: Deserialisation failed, please reimport the original model."), Ar.CustomVer(UE::NNE::ModelData::GUID));
 			break;
 		}
 	}
@@ -418,7 +418,7 @@ TSharedPtr<UE::NNE::FSharedModelData> UNNEModelData::GetModelData(const FString&
 	TArrayView<const FString> TargetRuntimesNames = GetTargetRuntimes();
 	if (!TargetRuntimesNames.IsEmpty() && !TargetRuntimesNames.Contains(RuntimeName))
 	{
-		UE_LOG(LogNNE, Error, TEXT("UNNEModelData: Runtime '%s' is not among the target runtimes. Target runtimes are: "), *RuntimeName);
+		UE_LOG(LogNNE, Error, TEXT("Runtime '%s' is not among the target runtimes. Target runtimes are: "), *RuntimeName);
 		for (const FString& TargetRuntimesName : TargetRuntimesNames)
 		{
 			UE_LOG(LogNNE, Error, TEXT("- %s"), *TargetRuntimesName);
@@ -436,7 +436,7 @@ TSharedPtr<UE::NNE::FSharedModelData> UNNEModelData::GetModelData(const FString&
 	// After this point FileData is required to either get the cache id or recreate it from scratch
 	if (FileData.Num() < 1)
 	{
-		UE_LOG(LogNNE, Error, TEXT("UNNEModelData: Cannot create model data from empty file data."));
+		UE_LOG(LogNNE, Error, TEXT("Cannot create model data from empty file data."));
 		return TSharedPtr<UE::NNE::FSharedModelData>();
 	}
 
