[NNE] NNERuntimeORTDml add tensor buffer size check to make sure to fullfill DirectML tensor buffer size requirements.

Problem:
DirectML has a 4-byte size alignment requirement for tensor buffers.

Solution:
Add tests in EnqueueRDG to catch buffer sizes violating size requirement.

Test:
Tested with "incorrect" buffer size and verified error log message instead of crash.

#tests Editor, Win64 verified error log instead of crash with incorrect buffer sizes

#rnx

diff --git a/Engine/Plugins/NNE/NNERuntimeORT/Source/NNERuntimeORT/Private/NNERuntimeORTModel.cpp b/Engine/Plugins/NNE/NNERuntimeORT/Source/NNERuntimeORT/Private/NNERuntimeORTModel.cpp
--- a/Engine/Plugins/NNE/NNERuntimeORT/Source/NNERuntimeORT/Private/NNERuntimeORTModel.cpp
+++ b/Engine/Plugins/NNE/NNERuntimeORT/Source/NNERuntimeORT/Private/NNERuntimeORTModel.cpp
@@ -901,9 +901,10 @@ FModelInstanceORTDmlRDG::EEnqueueRDGStatus FModelInstanceORTDmlRDG::EnqueueRDG(F
 			return EEnqueueRDGStatus::Fail;
 		}
 
-		if (Binding.Buffer->Desc.GetSize() != InputTensors[i].GetDataSize())
+		const uint64 DmlImpliedSizeBytes = CalcRDGBufferSizeForDirectML(InputTensors[i].GetDataSize());
+		if (Binding.Buffer && Binding.Buffer->Desc.GetSize() != DmlImpliedSizeBytes)
 		{
-			UE_LOG(LogNNERuntimeORT, Error, TEXT("Binding input tensor %d size does not match size given by tensor descriptor (got %d, expected %d)."), i, Binding.Buffer->Desc.GetSize(), InputTensors[i].GetDataSize());
+			UE_LOG(LogNNERuntimeORT, Error, TEXT("Binding input tensor %d size does not match tensor buffer size required by DirectML (got %d, expected %d, data size was %d)."), i, Binding.Buffer->Desc.GetSize(), DmlImpliedSizeBytes, InputTensors[i].GetDataSize());
 			return EEnqueueRDGStatus::Fail;
 		}
 
@@ -915,11 +916,15 @@ FModelInstanceORTDmlRDG::EEnqueueRDGStatus FModelInstanceORTDmlRDG::EnqueueRDG(F
 	{
 		const NNE::FTensorBindingRDG& Binding = Outputs[i];
 
-		if (Binding.Buffer && Binding.Buffer->Desc.GetSize() >= OutputTensors[i].GetDataSize())
+		const uint64 DmlImpliedSizeBytes = CalcRDGBufferSizeForDirectML(OutputTensors[i].GetDataSize());
+		if (Binding.Buffer && Binding.Buffer->Desc.GetSize() != DmlImpliedSizeBytes)
 		{
-			PassParameters->OutputBuffers.Emplace(Binding.Buffer, ERHIAccess::CopyDest);
-			ValidOutputs.Add(i);
+			UE_LOG(LogNNERuntimeORT, Error, TEXT("Binding output tensor %d size does not match tensor buffer size required by DirectML (got %d, expected %d, data size was %d)."), i, Binding.Buffer->Desc.GetSize(), DmlImpliedSizeBytes, OutputTensors[i].GetDataSize());
+			return EEnqueueRDGStatus::Fail;
 		}
+
+		PassParameters->OutputBuffers.Emplace(Binding.Buffer, ERHIAccess::CopyDest);
+		ValidOutputs.Add(i);
 	}
 
 	RDG_EVENT_SCOPE_STAT(GraphBuilder, FNNERuntimeORTDmlRDG, "FModelInstanceORTDmlRDG::EnqueueRDG");
diff --git a/Engine/Plugins/NNE/NNERuntimeORT/Source/NNERuntimeORT/Private/NNERuntimeORTUtils.cpp b/Engine/Plugins/NNE/NNERuntimeORT/Source/NNERuntimeORT/Private/NNERuntimeORTUtils.cpp
--- a/Engine/Plugins/NNE/NNERuntimeORT/Source/NNERuntimeORT/Private/NNERuntimeORTUtils.cpp
+++ b/Engine/Plugins/NNE/NNERuntimeORT/Source/NNERuntimeORT/Private/NNERuntimeORTUtils.cpp
@@ -445,4 +445,14 @@ TypeInfoORT TranslateTensorTypeORTToNNE(ONNXTensorElementDataType OrtDataType)
 	return TypeInfoORT{ DataType, ElementSize };
 }
 
+uint64 CalcRDGBufferSizeForDirectML(uint64 DataSize)
+{
+	uint64 MinimumImpliedSizeInBytes = DataSize;
+
+	// Round up to the nearest 4 bytes.
+	MinimumImpliedSizeInBytes = (MinimumImpliedSizeInBytes + 3) & ~3ull;
+
+	return MinimumImpliedSizeInBytes;
+}
+
 } // namespace UE::NNERuntimeORT::Private
\ No newline at end of file
diff --git a/Engine/Plugins/NNE/NNERuntimeORT/Source/NNERuntimeORT/Private/NNERuntimeORTUtils.h b/Engine/Plugins/NNE/NNERuntimeORT/Source/NNERuntimeORT/Private/NNERuntimeORTUtils.h
--- a/Engine/Plugins/NNE/NNERuntimeORT/Source/NNERuntimeORT/Private/NNERuntimeORTUtils.h
+++ b/Engine/Plugins/NNE/NNERuntimeORT/Source/NNERuntimeORT/Private/NNERuntimeORTUtils.h
@@ -39,4 +39,6 @@ namespace OrtHelper
 
 	TypeInfoORT TranslateTensorTypeORTToNNE(ONNXTensorElementDataType OrtDataType);
 
+	uint64 CalcRDGBufferSizeForDirectML(uint64 DataSize);
+
 } // UE::NNERuntimeORT::Private
\ No newline at end of file
