[NNE] NNERuntimeORT handle additional exceptions (SEH) triggered by ONNX Runtime (mainly by API calls to create session).

Problem:
During cook ONNX Runtime might crash with certain models during optimize in create session. Since these are not ISO-standard C++ exceptions, the process is terminated without logging the cause.

Solution:
On Windows perform structured exception handling (SEH) to catch errors from ORT and then skip model cook instead of crashing the cook process.

Tests:
Tested create model data in Editor and during cook and verified error logs instead of crash.

#tests create model data in Editor and cook

#rnx

diff --git a/Engine/Plugins/NNE/NNERuntimeORT/Source/NNERuntimeORT/Private/NNERuntimeORT.cpp b/Engine/Plugins/NNE/NNERuntimeORT/Source/NNERuntimeORT/Private/NNERuntimeORT.cpp
--- a/Engine/Plugins/NNE/NNERuntimeORT/Source/NNERuntimeORT/Private/NNERuntimeORT.cpp
+++ b/Engine/Plugins/NNE/NNERuntimeORT/Source/NNERuntimeORT/Private/NNERuntimeORT.cpp
@@ -89,7 +89,7 @@ TSharedPtr<UE::NNE::FSharedModelData> UNNERuntimeORTCpu::CreateModelData(const F
 
 	if (CanCreateModelData(FileType, FileData, AdditionalFileData, FileId, TargetPlatform) != ECanCreateModelDataStatus::Ok)
 	{
-		UE_LOG(LogNNERuntimeORT, Warning, TEXT("Cannot create the CPU model data with id %s (Filetype: %s)"), *FileId.ToString(EGuidFormats::Digits).ToLower(), *FileType);
+		UE_LOG(LogNNERuntimeORT, Error, TEXT("Cannot create the CPU model data with id %s (Filetype: %s)"), *FileId.ToString(EGuidFormats::Digits).ToLower(), *FileType);
 		return {};
 	}
 
@@ -108,6 +108,7 @@ TSharedPtr<UE::NNE::FSharedModelData> UNNERuntimeORTCpu::CreateModelData(const F
 
 			if (!OptimizeModel(Environment.ToSharedRef(), *SessionOptions, FileData, OptimizedModelBuffer))
 			{
+				UE_LOG(LogNNERuntimeORT, Error, TEXT("Failed to optimize model for CPU with id %s, model data will not be available"), *FileId.ToString(EGuidFormats::Digits).ToLower());
 				return {};
 			}
 
@@ -172,7 +173,7 @@ TSharedPtr<UE::NNE::IModelCPU> UNNERuntimeORTCpu::CreateModelCPU(const TObjectPt
 
 	if (CanCreateModelCPU(ModelData) != ECanCreateModelCPUStatus::Ok)
 	{
-		UE_LOG(LogNNERuntimeORT, Warning, TEXT("Cannot create a CPU model from the model data with id %s"), *ModelData->GetFileId().ToString(EGuidFormats::Digits));
+		UE_LOG(LogNNERuntimeORT, Error, TEXT("Cannot create a CPU model from the model data with id %s"), *ModelData->GetFileId().ToString(EGuidFormats::Digits));
 		return TSharedPtr<UE::NNE::IModelCPU>();
 	}
 
@@ -211,7 +212,7 @@ TSharedPtr<UE::NNE::FSharedModelData> UNNERuntimeORTDml::CreateModelData(const F
 
 	if (CanCreateModelData(FileType, FileData, AdditionalFileData, FileId, TargetPlatform) != ECanCreateModelDataStatus::Ok)
 	{
-		UE_LOG(LogNNERuntimeORT, Warning, TEXT("Cannot create the Dml model data with id %s (Filetype: %s)"), *FileId.ToString(EGuidFormats::Digits).ToLower(), *FileType);
+		UE_LOG(LogNNERuntimeORT, Error, TEXT("Cannot create the Dml model data with id %s (Filetype: %s)"), *FileId.ToString(EGuidFormats::Digits).ToLower(), *FileType);
 		return {};
 	}
 
@@ -230,6 +231,8 @@ TSharedPtr<UE::NNE::FSharedModelData> UNNERuntimeORTDml::CreateModelData(const F
 
 			if (!OptimizeModel(Environment.ToSharedRef(), *SessionOptions, FileData, OptimizedModelBuffer))
 			{
+				UE_LOG(LogNNERuntimeORT, Error, TEXT("Failed to optimize model for DirectML with id %s, model data will not be available"), *FileId.ToString(EGuidFormats::Digits).ToLower());
+
 				return {};
 			}
 
@@ -264,7 +267,7 @@ TSharedPtr<UE::NNE::IModelGPU> UNNERuntimeORTDml::CreateModelGPU(const TObjectPt
 
 	if (CanCreateModelGPU(ModelData) != ECanCreateModelGPUStatus::Ok)
 	{
-		UE_LOG(LogNNERuntimeORT, Warning, TEXT("Cannot create a GPU model from the model data with id %s"), *ModelData->GetFileId().ToString(EGuidFormats::Digits));
+		UE_LOG(LogNNERuntimeORT, Error, TEXT("Cannot create a GPU model from the model data with id %s"), *ModelData->GetFileId().ToString(EGuidFormats::Digits));
 		return {};
 	}
 
@@ -289,7 +292,7 @@ TSharedPtr<UE::NNE::IModelRDG> UNNERuntimeORTDml::CreateModelRDG(TObjectPtr<UNNE
 
 	if (CanCreateModelRDG(ModelData) != ECanCreateModelRDGStatus::Ok)
 	{
-		UE_LOG(LogNNERuntimeORT, Warning, TEXT("Cannot create a RDG model from the model data with id %s"), *ModelData->GetFileId().ToString(EGuidFormats::Digits));
+		UE_LOG(LogNNERuntimeORT, Error, TEXT("Cannot create a RDG model from the model data with id %s"), *ModelData->GetFileId().ToString(EGuidFormats::Digits));
 		return {};
 	}
 
diff --git a/Engine/Plugins/NNE/NNERuntimeORT/Source/NNERuntimeORT/Private/NNERuntimeORTModel.cpp b/Engine/Plugins/NNE/NNERuntimeORT/Source/NNERuntimeORT/Private/NNERuntimeORTModel.cpp
--- a/Engine/Plugins/NNE/NNERuntimeORT/Source/NNERuntimeORT/Private/NNERuntimeORTModel.cpp
+++ b/Engine/Plugins/NNE/NNERuntimeORT/Source/NNERuntimeORT/Private/NNERuntimeORTModel.cpp
@@ -2,8 +2,8 @@
 
 #include "NNERuntimeORTModel.h"
 
-#include "HAL/PlatformMisc.h"
 #include "HAL/PlatformFileManager.h"
+#include "HAL/PlatformMisc.h"
 #include "Misc/FileHelper.h"
 #include "Misc/Paths.h"
 #include "NNERuntimeORT.h"
@@ -121,15 +121,11 @@ namespace Detail
 				Filepath = FPaths::Combine(TempDirForModelWithExternalData, TEXT("OnnxModel.onnx"));
 			}
 
-#if PLATFORM_WINDOWS
-			Session = MakeUnique<Ort::Session>(Environment.GetOrtEnv(), *Filepath, SessionOptions);
-#else
-			Session = MakeUnique<Ort::Session>(Environment.GetOrtEnv(), TCHAR_TO_ANSI(*Filepath), SessionOptions);
-#endif
+			Session = CreateOrtSession(Environment, Filepath, SessionOptions);
 		}
 		else
 		{
-			Session = MakeUnique<Ort::Session>(Environment.GetOrtEnv(), ModelBuffer.GetData(), ModelBuffer.Num(), SessionOptions);
+			Session = CreateOrtSessionFromArray(Environment, ModelBuffer, SessionOptions);
 		}
 
 		return Session.IsValid();
@@ -763,28 +759,11 @@ FModelInstanceORTDmlRDG::ESetInputTensorShapesStatus FModelInstanceORTDmlRDG::Se
 		}
 	}
 
-#if WITH_EDITOR
-	try
-#endif // WITH_EDITOR
+	if (!Detail::CreateSession(ModelData->GetView(), *SessionOptions, *Environment, Session, TempDirForModelWithExternalData))
 	{
-		if (!Detail::CreateSession(ModelData->GetView(), *SessionOptions, *Environment, Session, TempDirForModelWithExternalData))
-		{
-			UE_LOG(LogNNERuntimeORT, Error, TEXT("Failed to recreate session!"));
-			return ESetInputTensorShapesStatus::Fail;
-		}
-	}
-#if WITH_EDITOR
-	catch (const Ort::Exception& Exception)
-	{
-		UE_LOG(LogNNERuntimeORT, Error, TEXT("%s"), UTF8_TO_TCHAR(Exception.what()));
+		UE_LOG(LogNNERuntimeORT, Error, TEXT("Failed to recreate session!"));
 		return ESetInputTensorShapesStatus::Fail;
 	}
-	catch (...)
-	{
-		UE_LOG(LogNNERuntimeORT, Error, TEXT("Unknown exception!"));
-		return ESetInputTensorShapesStatus::Fail;
-	}
-#endif // WITH_EDITOR
 
 	// Need to configure output tensors with new session (to apply free dimension overrides)
 	if (!ConfigureTensors(*Session, false))
diff --git a/Engine/Plugins/NNE/NNERuntimeORT/Source/NNERuntimeORT/Private/NNERuntimeORTUtils.cpp b/Engine/Plugins/NNE/NNERuntimeORT/Source/NNERuntimeORT/Private/NNERuntimeORTUtils.cpp
--- a/Engine/Plugins/NNE/NNERuntimeORT/Source/NNERuntimeORT/Private/NNERuntimeORTUtils.cpp
+++ b/Engine/Plugins/NNE/NNERuntimeORT/Source/NNERuntimeORT/Private/NNERuntimeORTUtils.cpp
@@ -1,6 +1,7 @@
 // Copyright Epic Games, Inc. All Rights Reserved.
 
 #include "NNERuntimeORTUtils.h"
+#include "HAL/ExceptionHandling.h"
 #include "HAL/IConsoleManager.h"
 #include "HAL/PlatformFileManager.h"
 #include "HAL/PlatformTime.h"
@@ -315,32 +316,23 @@ bool OptimizeModel(const TSharedRef<FEnvironment> &Environment, Ort::SessionOpti
 	FString ProjIntermediateDir = FPaths::ConvertRelativePathToFull(FPaths::ProjectIntermediateDir());
 	FString ModelOptimizedPath = FPaths::CreateTempFilename(*ProjIntermediateDir, TEXT("ORTOptimizerPass_Optimized"), TEXT(".onnx"));
 
-#if WITH_EDITOR
-	try
-#endif // WITH_EDITOR
-	{
 #if PLATFORM_WINDOWS
-		SessionOptions.SetOptimizedModelFilePath(*ModelOptimizedPath);
-
-		Ort::Session Session(Environment->GetOrtEnv(), InputModel.GetData(), InputModel.Num(), SessionOptions);
+	SessionOptions.SetOptimizedModelFilePath(*ModelOptimizedPath);
 #else
-		SessionOptions.SetOptimizedModelFilePath(TCHAR_TO_ANSI(*ModelOptimizedPath));
+	SessionOptions.SetOptimizedModelFilePath(StringCast<ANSICHAR>(*ModelOptimizedPath).Get());
+#endif // PLATFORM_WINDOWS
 
-		Ort::Session Session(Environment->GetOrtEnv(), InputModel.GetData(), InputModel.Num(), SessionOptions);
-#endif
-	}
-#if WITH_EDITOR
-	catch (const Ort::Exception& Exception)
 	{
-		UE_LOG(LogNNERuntimeORT, Error, TEXT("ORT Exception: %s"), UTF8_TO_TCHAR(Exception.what()));
-		return false;
-	}
-	catch (...)
-	{
-		UE_LOG(LogNNERuntimeORT, Error, TEXT("ORT Exception: Unknown!"));
-		return false;
+		TUniquePtr<Ort::Session> Session = CreateOrtSessionFromArray(Environment.Get(), InputModel, SessionOptions);
+		if (!Session)
+		{
+			UE_LOG(LogNNERuntimeORT, Error, TEXT("Failed to create ONNX Runtime session"));
+			
+			IFileManager::Get().Delete(*ModelOptimizedPath);
+			
+			return false;
+		}
 	}
-#endif // WITH_EDITOR
 
 	FFileHelper::LoadFileToArray(OptimizedModel, *ModelOptimizedPath);
 
@@ -455,4 +447,85 @@ uint64 CalcRDGBufferSizeForDirectML(uint64 DataSize)
 	return MinimumImpliedSizeInBytes;
 }
 
+template<typename FunctionType, typename... ArgTypes>
+bool OrtApiCallWithStatus(FunctionType &&Function, ArgTypes &&... Args)
+{
+	OrtStatusPtr StatusPtr = Forward<FunctionType>(Function)(Forward<ArgTypes>(Args)...);
+	if (StatusPtr)
+	{
+		OrtErrorCode Code = Ort::GetApi().GetErrorCode(StatusPtr);
+		const char* Message = Ort::GetApi().GetErrorMessage(StatusPtr);
+		UE_LOG(LogNNERuntimeORT, Error, TEXT("ONNX Runtime error %d: %hs"), (int32)Code, Message)
+		return false;
+	}
+	return true;
+}
+
+#if PLATFORM_WINDOWS && !PLATFORM_SEH_EXCEPTIONS_DISABLED
+template<typename FunctionType, typename... ArgTypes>
+bool GuardedOrtApiCallWithStatus(FunctionType &&Function, ArgTypes &&... Args)
+{
+	if (FPlatformMisc::IsDebuggerPresent())
+	{
+		return OrtApiCallWithStatus(Forward<FunctionType>(Function), Forward<ArgTypes>(Args)...);
+	}
+	else
+	{
+		__try
+		{
+			return OrtApiCallWithStatus(Forward<FunctionType>(Function), Forward<ArgTypes>(Args)...);
+		}
+		__except(EXCEPTION_EXECUTE_HANDLER)
+		{
+			UE_LOG(LogNNERuntimeORT, Error, TEXT("ONNX Runtime unknown exception (SEH)!"));
+			return false;
+		}
+	}
+}
+#endif // PLATFORM_WINDOWS && !PLATFORM_SEH_EXCEPTIONS_DISABLED
+
+TUniquePtr<Ort::Session> CreateOrtSessionFromArray(const FEnvironment& Environment, TConstArrayView64<uint8> ModelBuffer, const Ort::SessionOptions& SessionOptions)
+{
+	OrtSession* SessionPtr = nullptr;
+#if PLATFORM_WINDOWS && !PLATFORM_SEH_EXCEPTIONS_DISABLED
+	bool Success = GuardedOrtApiCallWithStatus(Ort::GetApi().CreateSessionFromArray, Environment.GetOrtEnv(), ModelBuffer.GetData(), ModelBuffer.Num(), SessionOptions, &SessionPtr);
+#else
+	bool Success = OrtApiCallWithStatus(Ort::GetApi().CreateSessionFromArray, Environment.GetOrtEnv(), ModelBuffer.GetData(), ModelBuffer.Num(), SessionOptions, &SessionPtr);
+#endif // PLATFORM_WINDOWS && !PLATFORM_SEH_EXCEPTIONS_DISABLED
+	if (!Success)
+	{
+		return {};
+	}
+
+	check(SessionPtr);
+
+	TUniquePtr<Ort::Session> Session;
+	Session.Reset(static_cast<Ort::Session *>(new Ort::Session::Base(SessionPtr)));
+
+	return Session;
+}
+
+TUniquePtr<Ort::Session> CreateOrtSession(const FEnvironment& Environment, const FString& ModelPath, const Ort::SessionOptions& SessionOptions)
+{
+	OrtSession* SessionPtr = nullptr;
+#if PLATFORM_WINDOWS && !PLATFORM_SEH_EXCEPTIONS_DISABLED
+	bool Success = GuardedOrtApiCallWithStatus(Ort::GetApi().CreateSession, Environment.GetOrtEnv(), *ModelPath, SessionOptions, &SessionPtr);
+#elif PLATFORM_WINDOWS
+	bool Success = OrtApiCallWithStatus(Ort::GetApi().CreateSession, Environment.GetOrtEnv(), *ModelPath, SessionOptions, &SessionPtr);
+#else
+	bool Success = OrtApiCallWithStatus(Ort::GetApi().CreateSession, Environment.GetOrtEnv(), StringCast<ANSICHAR>(*ModelPath).Get(), SessionOptions, &SessionPtr);
+#endif
+	if (!Success)
+	{
+		return {};
+	}
+
+	check(SessionPtr);
+	
+	TUniquePtr<Ort::Session> Session;
+	Session.Reset(static_cast<Ort::Session *>(new Ort::Session::Base(SessionPtr)));
+
+	return Session;
+}
+
 } // namespace UE::NNERuntimeORT::Private
\ No newline at end of file
diff --git a/Engine/Plugins/NNE/NNERuntimeORT/Source/NNERuntimeORT/Private/NNERuntimeORTUtils.h b/Engine/Plugins/NNE/NNERuntimeORT/Source/NNERuntimeORT/Private/NNERuntimeORTUtils.h
--- a/Engine/Plugins/NNE/NNERuntimeORT/Source/NNERuntimeORT/Private/NNERuntimeORTUtils.h
+++ b/Engine/Plugins/NNE/NNERuntimeORT/Source/NNERuntimeORT/Private/NNERuntimeORTUtils.h
@@ -41,4 +41,7 @@ namespace OrtHelper
 
 	uint64 CalcRDGBufferSizeForDirectML(uint64 DataSize);
 
+	TUniquePtr<Ort::Session> CreateOrtSessionFromArray(const FEnvironment& Environment, TConstArrayView64<uint8> ModelBuffer, const Ort::SessionOptions& SessionOptions);
+	TUniquePtr<Ort::Session> CreateOrtSession(const FEnvironment& Environment, const FString& ModelPath, const Ort::SessionOptions& SessionOptions);
+
 } // UE::NNERuntimeORT::Private
\ No newline at end of file
