[NNE] NNERuntimeORTDml implements NPU interface.

#tests Packaged and run on machine with NPU device

diff --git a/Engine/Plugins/NNE/NNERuntimeORT/Source/NNERuntimeORT/Private/NNERuntimeORT.cpp b/Engine/Plugins/NNE/NNERuntimeORT/Source/NNERuntimeORT/Private/NNERuntimeORT.cpp
--- a/Engine/Plugins/NNE/NNERuntimeORT/Source/NNERuntimeORT/Private/NNERuntimeORT.cpp
+++ b/Engine/Plugins/NNE/NNERuntimeORT/Source/NNERuntimeORT/Private/NNERuntimeORT.cpp
@@ -177,13 +177,9 @@ TSharedPtr<UE::NNE::IModelCPU> UNNERuntimeORTCpu::CreateModelCPU(const TObjectPt
 		return TSharedPtr<UE::NNE::IModelCPU>();
 	}
 
-	const TSharedPtr<UE::NNE::FSharedModelData> SharedData = ModelData->GetModelData(GetRuntimeName());
-	check(SharedData.IsValid());
-
-	UE::NNE::IModelCPU* IModel = static_cast<UE::NNE::IModelCPU*>(new UE::NNERuntimeORT::Private::FModelORTCpu(Environment.ToSharedRef(), SharedData));
-	check(IModel != nullptr);
+	const TSharedRef<UE::NNE::FSharedModelData> SharedData = ModelData->GetModelData(GetRuntimeName()).ToSharedRef();
 
-	return TSharedPtr<UE::NNE::IModelCPU>(IModel);
+	return MakeShared<UE::NNERuntimeORT::Private::FModelORTCpu>(Environment.ToSharedRef(), SharedData);
 }
 
 /*
@@ -194,6 +190,7 @@ void UNNERuntimeORTDml::Init(TSharedRef<UE::NNERuntimeORT::Private::FEnvironment
 	Environment = InEnvironment;
 	bDirectMLAvailable = bInDirectMLAvailable;
 	bD3D12Available = UE::NNERuntimeORT::Private::IsD3D12Available();
+	bD3D12DeviceNPUAvailable = UE::NNERuntimeORT::Private::IsD3D12DeviceNPUAvailable();
 }
 
 FString UNNERuntimeORTDml::GetRuntimeName() const
@@ -257,6 +254,16 @@ FString UNNERuntimeORTDml::GetModelDataIdentifier(const FString& FileType, TCons
 
 UNNERuntimeORTDml::ECanCreateModelGPUStatus UNNERuntimeORTDml::CanCreateModelGPU(const TObjectPtr<UNNEModelData> ModelData) const
 {
+	if (!bDirectMLAvailable)
+	{
+		return ECanCreateModelCommonStatus::Fail;
+	}
+
+	if (!bD3D12Available)
+	{
+		return ECanCreateModelCommonStatus::Fail;
+	}
+
 	return CanCreateModelCommon(ModelData, false) == ECanCreateModelCommonStatus::Ok ? ECanCreateModelGPUStatus::Ok : ECanCreateModelGPUStatus::Fail;
 }
 
@@ -271,8 +278,7 @@ TSharedPtr<UE::NNE::IModelGPU> UNNERuntimeORTDml::CreateModelGPU(const TObjectPt
 		return {};
 	}
 
-	const TSharedPtr<UE::NNE::FSharedModelData> SharedData = ModelData->GetModelData(GetRuntimeName());
-	check(SharedData.IsValid());
+	const TSharedRef<UE::NNE::FSharedModelData> SharedData = ModelData->GetModelData(GetRuntimeName()).ToSharedRef();
 
 	return MakeShared<UE::NNERuntimeORT::Private::FModelORTDmlGPU>(Environment.ToSharedRef(), SharedData);
 #else // PLATFORM_WINDOWS
@@ -282,6 +288,18 @@ TSharedPtr<UE::NNE::IModelGPU> UNNERuntimeORTDml::CreateModelGPU(const TObjectPt
 
 UNNERuntimeORTDml::ECanCreateModelRDGStatus UNNERuntimeORTDml::CanCreateModelRDG(TObjectPtr<UNNEModelData> ModelData) const
 {
+	if (!bDirectMLAvailable)
+	{
+		return ECanCreateModelCommonStatus::Fail;
+	}
+
+#if PLATFORM_WINDOWS
+	if (!IsRHID3D12())
+	{
+		return ECanCreateModelCommonStatus::Fail;
+	}
+#endif // PLATFORM_WINDOWS
+
 	return CanCreateModelCommon(ModelData) == ECanCreateModelCommonStatus::Ok ? ECanCreateModelRDGStatus::Ok : ECanCreateModelRDGStatus::Fail;
 }
 
@@ -304,27 +322,45 @@ TSharedPtr<UE::NNE::IModelRDG> UNNERuntimeORTDml::CreateModelRDG(TObjectPtr<UNNE
 #endif // PLATFORM_WINDOWS
 }
 
-UNNERuntimeORTDml::ECanCreateModelCommonStatus UNNERuntimeORTDml::CanCreateModelCommon(const TObjectPtr<UNNEModelData> ModelData, bool bRHID3D12Required) const
+UNNERuntimeORTDml::ECanCreateModelNPUStatus UNNERuntimeORTDml::CanCreateModelNPU(const TObjectPtr<UNNEModelData> ModelData) const
 {
-#if PLATFORM_WINDOWS
-	check(ModelData != nullptr);
-
-	// DirectML is required
 	if (!bDirectMLAvailable)
 	{
 		return ECanCreateModelCommonStatus::Fail;
 	}
 
-	// Either RHID3D12 or at least D3D12 is required
-	if (bRHID3D12Required && !IsRHID3D12())
+	if (!bD3D12DeviceNPUAvailable)
 	{
 		return ECanCreateModelCommonStatus::Fail;
 	}
-	else if (!bD3D12Available)
+
+	return CanCreateModelCommon(ModelData) == ECanCreateModelCommonStatus::Ok ? ECanCreateModelRDGStatus::Ok : ECanCreateModelRDGStatus::Fail;
+}
+
+TSharedPtr<UE::NNE::IModelNPU> UNNERuntimeORTDml::CreateModelNPU(const TObjectPtr<UNNEModelData> ModelData)
+{
+#if PLATFORM_WINDOWS
+	check(ModelData);
+
+	if (CanCreateModelNPU(ModelData) != ECanCreateModelRDGStatus::Ok)
 	{
-		return ECanCreateModelCommonStatus::Fail;
+		UE_LOG(LogNNERuntimeORT, Error, TEXT("Cannot create a model NPU from the model data with id %s"), *ModelData->GetFileId().ToString(EGuidFormats::Digits));
+		return {};
 	}
 
+	const TSharedRef<UE::NNE::FSharedModelData> SharedData = ModelData->GetModelData(GetRuntimeName()).ToSharedRef();
+
+	return MakeShared<UE::NNERuntimeORT::Private::FModelORTNpu>(Environment.ToSharedRef(), SharedData);
+#else // PLATFORM_WINDOWS
+	return {};
+#endif // PLATFORM_WINDOWS
+}
+
+UNNERuntimeORTDml::ECanCreateModelCommonStatus UNNERuntimeORTDml::CanCreateModelCommon(const TObjectPtr<UNNEModelData> ModelData, bool bRHID3D12Required) const
+{
+#if PLATFORM_WINDOWS
+	check(ModelData != nullptr);
+
 	constexpr int32 GuidSize = sizeof(UNNERuntimeORTDml::GUID);
 	constexpr int32 VersionSize = sizeof(UNNERuntimeORTDml::Version);
 	const TSharedPtr<UE::NNE::FSharedModelData> SharedData = ModelData->GetModelData(GetRuntimeName());
diff --git a/Engine/Plugins/NNE/NNERuntimeORT/Source/NNERuntimeORT/Private/NNERuntimeORT.h b/Engine/Plugins/NNE/NNERuntimeORT/Source/NNERuntimeORT/Private/NNERuntimeORT.h
--- a/Engine/Plugins/NNE/NNERuntimeORT/Source/NNERuntimeORT/Private/NNERuntimeORT.h
+++ b/Engine/Plugins/NNE/NNERuntimeORT/Source/NNERuntimeORT/Private/NNERuntimeORT.h
@@ -4,6 +4,7 @@
 
 #include "NNERuntime.h"
 #include "NNERuntimeCPU.h"
+#include "NNERuntimeNPU.h"
 #include "NNERuntimeGPU.h"
 #include "NNERuntimeRDG.h"
 #include "UObject/Class.h"
@@ -46,7 +47,7 @@ public:
 };
 
 UCLASS()
-class UNNERuntimeORTDml : public UObject, public INNERuntime, public INNERuntimeGPU, public INNERuntimeRDG
+class UNNERuntimeORTDml : public UObject, public INNERuntime, public INNERuntimeGPU, public INNERuntimeRDG, public INNERuntimeNPU
 {
 	GENERATED_BODY()
 
@@ -75,9 +76,13 @@ public:
 	virtual ECanCreateModelRDGStatus CanCreateModelRDG(TObjectPtr<UNNEModelData> ModelData) const override;
 	virtual TSharedPtr<UE::NNE::IModelRDG> CreateModelRDG(TObjectPtr<UNNEModelData> ModelData) override;
 
+	virtual ECanCreateModelNPUStatus CanCreateModelNPU(const TObjectPtr<UNNEModelData> ModelData) const override;
+	virtual TSharedPtr<UE::NNE::IModelNPU> CreateModelNPU(const TObjectPtr<UNNEModelData> ModelData) override;
+
 private:
 	ECanCreateModelCommonStatus CanCreateModelCommon(const TObjectPtr<UNNEModelData> ModelData, bool bRHID3D12Required = true) const;
 
 	bool bDirectMLAvailable = false;
 	bool bD3D12Available = false;
+	bool bD3D12DeviceNPUAvailable = false;
 };
\ No newline at end of file
diff --git a/Engine/Plugins/NNE/NNERuntimeORT/Source/NNERuntimeORT/Private/NNERuntimeORTModel.cpp b/Engine/Plugins/NNE/NNERuntimeORT/Source/NNERuntimeORT/Private/NNERuntimeORTModel.cpp
--- a/Engine/Plugins/NNE/NNERuntimeORT/Source/NNERuntimeORT/Private/NNERuntimeORTModel.cpp
+++ b/Engine/Plugins/NNE/NNERuntimeORT/Source/NNERuntimeORT/Private/NNERuntimeORTModel.cpp
@@ -431,21 +431,18 @@ typename ModelInterface::ERunSyncStatus FModelInstanceORTBase<ModelInterface, Te
 
 TSharedPtr<NNE::IModelInstanceCPU> FModelORTCpu::CreateModelInstanceCPU()
 {
-	const FRuntimeConf Conf = Detail::MakeRuntimeConfigFromSettings(GetDefault<UNNERuntimeORTSettings>());
-	FModelInstanceORTCpu* ModelInstance = new FModelInstanceORTCpu(Conf, Environment);
+	const FRuntimeConf RuntimeConfig = Detail::MakeRuntimeConfigFromSettings(GetDefault<UNNERuntimeORTSettings>());
 
-	check(ModelData.IsValid());
+	TSharedPtr<FModelInstanceORTCpu> ModelInstance = MakeShared<FModelInstanceORTCpu>(RuntimeConfig, Environment);
 	if (!ModelInstance->Init(ModelData->GetView()))
 	{
-		delete ModelInstance;
-		return TSharedPtr<UE::NNE::IModelInstanceCPU>();
+		return {};
 	}
 
-	NNE::IModelInstanceCPU* IModelInstance = static_cast<NNE::IModelInstanceCPU*>(ModelInstance);
-	return TSharedPtr<NNE::IModelInstanceCPU>(IModelInstance);
+	return ModelInstance;
 }
 
-FModelORTCpu::FModelORTCpu(TSharedRef<FEnvironment> InEnvironment, const TSharedPtr<UE::NNE::FSharedModelData>& InModelData) :
+FModelORTCpu::FModelORTCpu(TSharedRef<FEnvironment> InEnvironment, TSharedRef<UE::NNE::FSharedModelData> InModelData) :
 	Environment(InEnvironment), ModelData(InModelData)
 {
 }
@@ -475,20 +472,16 @@ TSharedPtr<NNE::IModelInstanceGPU> FModelORTDmlGPU::CreateModelInstanceGPU()
 {
 	const FRuntimeConf RuntimeConfig = Detail::MakeRuntimeConfigFromSettings(GetDefault<UNNERuntimeORTSettings>());
 
-	FModelInstanceORTDmlGPU* ModelInstance = new FModelInstanceORTDmlGPU(RuntimeConfig, Environment);
-
-	check(ModelData.IsValid());
+	TSharedPtr<FModelInstanceORTDmlGPU> ModelInstance = MakeShared<FModelInstanceORTDmlGPU>(RuntimeConfig, Environment);
 	if (!ModelInstance->Init(ModelData->GetView()))
 	{
-		delete ModelInstance;
-		return TSharedPtr<UE::NNE::IModelInstanceGPU>();
+		return {};
 	}
 
-	NNE::IModelInstanceGPU* IModelInstance = static_cast<NNE::IModelInstanceGPU*>(ModelInstance);
-	return TSharedPtr<NNE::IModelInstanceGPU>(IModelInstance);
+	return ModelInstance;
 }
 
-FModelORTDmlGPU::FModelORTDmlGPU(TSharedRef<FEnvironment> InEnvironment, const TSharedPtr<UE::NNE::FSharedModelData>& InModelData) :
+FModelORTDmlGPU::FModelORTDmlGPU(TSharedRef<FEnvironment> InEnvironment, TSharedRef<UE::NNE::FSharedModelData> InModelData) :
 	Environment(InEnvironment), ModelData(InModelData)
 {
 }
@@ -986,6 +979,43 @@ FModelInstanceORTDmlRDG::EEnqueueRDGStatus FModelInstanceORTDmlRDG::EnqueueRDG(F
 
 	return EEnqueueRDGStatus::Ok;
 }
+
+TSharedPtr<NNE::IModelInstanceNPU> FModelORTNpu::CreateModelInstanceNPU()
+{
+	const FRuntimeConf RuntimeConfig = Detail::MakeRuntimeConfigFromSettings(GetDefault<UNNERuntimeORTSettings>());
+
+	TSharedPtr<FModelInstanceORTNpu> ModelInstance = MakeShared<FModelInstanceORTNpu>(RuntimeConfig, Environment);
+	if (!ModelInstance->Init(ModelData->GetView()))
+	{
+		return {};
+	}
+
+	return ModelInstance;
+}
+
+FModelORTNpu::FModelORTNpu(TSharedRef<FEnvironment> InEnvironment, TSharedRef<UE::NNE::FSharedModelData> InModelData) :
+	Environment(InEnvironment), ModelData(InModelData)
+{
+}
+
+bool FModelInstanceORTNpu::InitializedAndConfigureMembers()
+{
+	if (!FModelInstanceORTBase::InitializedAndConfigureMembers())
+	{
+		return false;
+	}
+
+	SessionOptions = CreateSessionOptionsForDirectMLNpu(Environment);
+	if (!SessionOptions.IsValid())
+	{
+		return false;
+	}
+
+	SessionOptions->SetExecutionMode(RuntimeConf.ExecutionMode);
+	SessionOptions->SetGraphOptimizationLevel(GetGraphOptimizationLevelForDML(true));
+
+	return true;
+}
 #endif //PLATFORM_WINDOWS
 	
 } // namespace UE::NNERuntimeORT::Private
\ No newline at end of file
diff --git a/Engine/Plugins/NNE/NNERuntimeORT/Source/NNERuntimeORT/Private/NNERuntimeORTModel.h b/Engine/Plugins/NNE/NNERuntimeORT/Source/NNERuntimeORT/Private/NNERuntimeORTModel.h
--- a/Engine/Plugins/NNE/NNERuntimeORT/Source/NNERuntimeORT/Private/NNERuntimeORTModel.h
+++ b/Engine/Plugins/NNE/NNERuntimeORT/Source/NNERuntimeORT/Private/NNERuntimeORTModel.h
@@ -7,6 +7,7 @@
 #include "NNEModelData.h"
 #include "NNEOnnxruntime.h"
 #include "NNERuntimeCPU.h"
+#include "NNERuntimeNPU.h"
 #include "NNERuntimeGPU.h"
 #include "NNERuntimeORTEnv.h"
 #include "NNERuntimeRDG.h"
@@ -75,28 +76,28 @@ private:
 class FModelORTCpu : public NNE::IModelCPU
 {
 public:
-	FModelORTCpu(TSharedRef<FEnvironment> InEnvironment, const TSharedPtr<UE::NNE::FSharedModelData>& InModelData);
+	FModelORTCpu(TSharedRef<FEnvironment> InEnvironment, TSharedRef<UE::NNE::FSharedModelData> InModelData);
 	virtual ~FModelORTCpu() = default;
 
 	virtual TSharedPtr<NNE::IModelInstanceCPU> CreateModelInstanceCPU() override;
 
 private:
 	TSharedRef<FEnvironment> Environment;
-	TSharedPtr<UE::NNE::FSharedModelData> ModelData;
+	TSharedRef<UE::NNE::FSharedModelData> ModelData;
 };
 
 #if PLATFORM_WINDOWS
 class FModelORTDmlGPU : public NNE::IModelGPU
 {
 public:
-	FModelORTDmlGPU(TSharedRef<FEnvironment> InEnvironment, const TSharedPtr<UE::NNE::FSharedModelData>& InModelData);
+	FModelORTDmlGPU(TSharedRef<FEnvironment> InEnvironment, TSharedRef<UE::NNE::FSharedModelData> InModelData);
 	virtual ~FModelORTDmlGPU() = default;
 
 	virtual TSharedPtr<NNE::IModelInstanceGPU> CreateModelInstanceGPU() override;
 
 private:
 	TSharedRef<FEnvironment> Environment;
-	TSharedPtr<UE::NNE::FSharedModelData> ModelData;
+	TSharedRef<UE::NNE::FSharedModelData> ModelData;
 };
 
 class FModelInstanceORTDmlGPU : public FModelInstanceORTBase<NNE::IModelInstanceGPU, NNE::FTensorBindingCPU>
@@ -168,6 +169,29 @@ protected:
 	TArray<NNE::Internal::FTensor> InputTensors;
 	TArray<NNE::Internal::FTensor> OutputTensors;
 };
+
+class FModelInstanceORTNpu : public FModelInstanceORTBase<NNE::IModelInstanceNPU, NNE::FTensorBindingCPU>
+{
+public:
+	FModelInstanceORTNpu(const FRuntimeConf& InRuntimeConf, TSharedRef<FEnvironment> InEnvironment) : FModelInstanceORTBase(InRuntimeConf, InEnvironment) {}
+	virtual ~FModelInstanceORTNpu() = default;
+
+private:
+	virtual bool InitializedAndConfigureMembers() override;
+};
+
+class FModelORTNpu : public NNE::IModelNPU
+{
+public:
+	FModelORTNpu(TSharedRef<FEnvironment> InEnvironment, TSharedRef<UE::NNE::FSharedModelData> InModelData);
+	virtual ~FModelORTNpu() = default;
+
+	virtual TSharedPtr<NNE::IModelInstanceNPU> CreateModelInstanceNPU() override;
+
+private:
+	TSharedRef<FEnvironment> Environment;
+	TSharedRef<UE::NNE::FSharedModelData> ModelData;
+};
 #endif //PLATFORM_WINDOWS
 	
 } // UE::NNERuntimeORT::Private
\ No newline at end of file
diff --git a/Engine/Plugins/NNE/NNERuntimeORT/Source/NNERuntimeORT/Private/NNERuntimeORTUtils.cpp b/Engine/Plugins/NNE/NNERuntimeORT/Source/NNERuntimeORT/Private/NNERuntimeORTUtils.cpp
--- a/Engine/Plugins/NNE/NNERuntimeORT/Source/NNERuntimeORT/Private/NNERuntimeORTUtils.cpp
+++ b/Engine/Plugins/NNE/NNERuntimeORT/Source/NNERuntimeORT/Private/NNERuntimeORTUtils.cpp
@@ -11,17 +11,11 @@
 #include "NNERuntimeORTEnv.h"
 
 #if PLATFORM_WINDOWS
-#include <dxgi1_4.h>
+#include <dxcore_interface.h>
+#include <dxcore.h>
 #include "ID3D12DynamicRHI.h"
 #endif // PLATFORM_WINDOWS
 
-// DirectML is implemented using COM on all platforms
-#ifdef IID_GRAPHICS_PPV_ARGS
-#define DML_PPV_ARGS(x) __uuidof(*x), IID_PPV_ARGS_Helper(x)
-#else
-#define DML_PPV_ARGS(x) IID_PPV_ARGS(x)
-#endif
-
 static int32 ORTProfilingSessionNumber = 0;
 static TAutoConsoleVariable<bool> CVarNNERuntimeORTEnableProfiling(
 	TEXT("nne.ort.enableprofiling"),
@@ -34,10 +28,40 @@ static TAutoConsoleVariable<bool> CVarNNERuntimeORTEnableProfiling(
 
 namespace UE::NNERuntimeORT::Private
 {
+
+#if PLATFORM_WINDOWS
+Microsoft::WRL::ComPtr<ID3D12Device1> CreateD3D12Device(IUnknown* AdapterPtr)
+{
+	using Microsoft::WRL::ComPtr;
+
+	void* D3D12Module = FPlatformProcess::GetDllHandle(TEXT("d3d12.dll"));
+	if (!D3D12Module)
+	{
+		UE_LOG(LogNNERuntimeORT, Error, TEXT("Failed to load module 'd3d12.dll'"));
+		return {};
+	}
+
+	decltype(&D3D12CreateDevice) D3D12CreateDeviceFun = reinterpret_cast<decltype(&D3D12CreateDevice)>(FPlatformProcess::GetDllExport(D3D12Module, TEXT("D3D12CreateDevice")));
+	if (!D3D12CreateDeviceFun)
+	{
+		UE_LOG(LogNNERuntimeORT, Error, TEXT("Failed to get export 'D3D12CreateDevice' from module 'd3d12.dll'"));
+		return {};
+	}
+
+	ComPtr<ID3D12Device1> D3D12Device;
+	HRESULT Hr = D3D12CreateDeviceFun(AdapterPtr, D3D_FEATURE_LEVEL_1_0_CORE, IID_PPV_ARGS(&D3D12Device));
+	if (FAILED(Hr))
+	{
+		UE_LOG(LogNNERuntimeORT, Error, TEXT("Failed to create D3D12 device, D3D12CreateDevice error code :%x"), Hr);
+		return {};
+	}
+
+	return D3D12Device;
+}
+#endif // PLATFORM_WINDOWS
+
 // Check for DirectX 12-compatible hardware.
-// Use DXGI to enumerate adapters and try to create a d3d12 device using the default adapter (will create dependency to dxgi.dll!)
-// DXGI 1.6 should be available since Windows 10, version 1809, which is newer than the minimum SDK version
-// specified in Engine\Config\Windows\Windows_SDK.json at the moment.
+// Use DXGI to enumerate adapters and try to create a d3d12 device using the default adapter
 bool IsD3D12Available()
 {
 #if PLATFORM_WINDOWS
@@ -45,22 +69,51 @@ bool IsD3D12Available()
 
 	const int32 DeviceIndex = 0;
 
-	ComPtr<IDXGIFactory4> Factory;
-	CreateDXGIFactory2(0, IID_PPV_ARGS(&Factory));
-	if (!Factory)
+	void* DxCoreModule = FPlatformProcess::GetDllHandle(TEXT("DXCore.dll"));
+	if (!DxCoreModule)
 	{
 		return false;
 	}
 
-	ComPtr<IDXGIAdapter1> Adapter;
-	Factory->EnumAdapters1(DeviceIndex, &Adapter);
-	if (!Adapter)
+	using DXCoreCreateAdapterFactoryFunType = HRESULT __stdcall(REFIID, void**);
+
+	DXCoreCreateAdapterFactoryFunType* DxCoreCreateAdapterFactoryFun = reinterpret_cast<DXCoreCreateAdapterFactoryFunType*>(FPlatformProcess::GetDllExport(DxCoreModule, TEXT("DXCoreCreateAdapterFactory")));
+	if (!DxCoreCreateAdapterFactoryFun)
+	{
+		return false;
+	}
+
+	ComPtr<IDXCoreAdapterFactory> Factory;
+	HRESULT Hr = DxCoreCreateAdapterFactoryFun(IID_PPV_ARGS(&Factory));
+	if (FAILED(Hr))
+	{
+		return false;
+	}
+
+	const GUID DxGUIDs[] = { DXCORE_ADAPTER_ATTRIBUTE_D3D12_CORE_COMPUTE };
+
+	ComPtr<IDXCoreAdapterList> AdapterList;
+	Hr = Factory->CreateAdapterList(ARRAYSIZE(DxGUIDs), DxGUIDs, IID_PPV_ARGS(&AdapterList));
+	if (FAILED(Hr))
 	{
 		return false;
 	}
 
-	ComPtr<ID3D12Device> Device;
-	D3D12CreateDevice(Adapter.Get(), D3D_FEATURE_LEVEL_11_0, IID_PPV_ARGS(&Device));
+	const uint32 AdapterCount = AdapterList->GetAdapterCount();
+	if (AdapterCount <= DeviceIndex)
+	{
+		UE_LOG(LogNNERuntimeORT, Error, TEXT("Invalid device index %d. Number of available devices is %d."), DeviceIndex, AdapterCount);
+		return false;
+	}
+
+	ComPtr<IDXCoreAdapter> Adapter;
+	Hr = AdapterList->GetAdapter(static_cast<uint32_t>(DeviceIndex), IID_PPV_ARGS(&Adapter));
+	if (FAILED(Hr))
+	{
+		return false;
+	}
+
+	ComPtr<ID3D12Device> Device = CreateD3D12Device(Adapter.Get());
 	if (!Device)
 	{
 		return false;
@@ -72,6 +125,152 @@ bool IsD3D12Available()
 #endif // PLATFORM_WINDOWS
 }
 
+#if PLATFORM_WINDOWS
+Microsoft::WRL::ComPtr<IUnknown> GetAdapterNpu(bool bVerbose)
+{
+	using Microsoft::WRL::ComPtr;
+
+	const bool bForceComputeOnlyDevice = true;
+
+	auto GetAdapterName = [](IDXCoreAdapter* Adapter) -> FString
+		{
+			DXCoreAdapterProperty DriverDescProperty = DXCoreAdapterProperty::DriverDescription;
+
+			if (!Adapter->IsPropertySupported(DriverDescProperty))
+			{
+				return TEXT("Unknown device");
+			}
+
+			size_t PropertySize = 0;
+			HRESULT Hr = Adapter->GetPropertySize(DriverDescProperty, &PropertySize);
+			if (FAILED(Hr))
+			{
+				return TEXT("Unknown device");
+			}
+
+			TArray<char> Buffer;
+			Buffer.Init(0, (int32)PropertySize);
+
+			Hr = Adapter->GetProperty(DriverDescProperty, PropertySize, Buffer.GetData());
+			if (FAILED(Hr))
+			{
+				return TEXT("Unknown device");
+			}
+			
+			return FString(ANSI_TO_TCHAR(Buffer.GetData()));
+		};
+
+	void* DxCoreModule = FPlatformProcess::GetDllHandle(TEXT("DXCore.dll"));
+	if (!DxCoreModule)
+	{
+		UE_LOG(LogNNERuntimeORT, Error, TEXT("Failed to load module 'DXCore.dll'"));
+		return {};
+	}
+
+	using DXCoreCreateAdapterFactoryFunType = HRESULT __stdcall(REFIID, void**);
+
+	DXCoreCreateAdapterFactoryFunType* DxCoreCreateAdapterFactoryFun = reinterpret_cast<DXCoreCreateAdapterFactoryFunType*>(FPlatformProcess::GetDllExport(DxCoreModule, TEXT("DXCoreCreateAdapterFactory")));
+	if (!DxCoreCreateAdapterFactoryFun)
+	{
+		UE_LOG(LogNNERuntimeORT, Error, TEXT("Failed to get export 'DXCoreCreateAdapterFactory' from module 'DXCore.dll'"));
+		return {};
+	}
+
+	ComPtr<IDXCoreAdapterFactory> Factory;
+	HRESULT Hr = DxCoreCreateAdapterFactoryFun(IID_PPV_ARGS(&Factory));
+	if (FAILED(Hr))
+	{
+		UE_LOG(LogNNERuntimeORT, Error, TEXT("Failed to create DXCore Adapter Factory, DxCoreCreateAdapterFactory error code :%x"), Hr);
+		return {};
+	}
+
+	const GUID DxGUIDs[] = { DXCORE_ADAPTER_ATTRIBUTE_D3D12_CORE_COMPUTE };
+
+	ComPtr<IDXCoreAdapterList> AdapterList;
+	Hr = Factory->CreateAdapterList(ARRAYSIZE(DxGUIDs), DxGUIDs, IID_PPV_ARGS(&AdapterList));
+	if (FAILED(Hr))
+	{
+		UE_LOG(LogNNERuntimeORT, Error, TEXT("Failed to create DXCore Adapter List, IDXCoreAdapterFactory::CreateAdapterList error code :%x"), Hr);
+		return {};
+	}
+
+	if (bVerbose)
+	{
+		UE_LOG(LogNNERuntimeORT, Log, TEXT("Available graphics and compute adapters:"));
+	}
+
+	ComPtr<IDXCoreAdapter> Adapter;
+	for (uint32_t i = 0, AdapterCount = AdapterList->GetAdapterCount(); i < AdapterCount; i++)
+	{
+		ComPtr<IDXCoreAdapter> CurrentAdapter;
+		Hr = AdapterList->GetAdapter(static_cast<uint32_t>(i), IID_PPV_ARGS(&CurrentAdapter));
+		if (FAILED(Hr))
+		{
+			UE_LOG(LogNNERuntimeORT, Warning, TEXT("%d: Failed to get adapter, IDXCoreAdapterList::GetAdapter error code: %x"), i, Hr);
+			continue;
+		}
+
+		const bool bIsGraphicsAdapter = CurrentAdapter->IsAttributeSupported(DXCORE_ADAPTER_ATTRIBUTE_D3D12_GRAPHICS);
+
+		if (bVerbose)
+		{
+			TArray<FString> SupportedAttributesStringList;
+			SupportedAttributesStringList.Add("Compute");
+			if (bIsGraphicsAdapter)
+			{
+				SupportedAttributesStringList.Add("Graphics");
+			}
+
+			FString SupportedAttributesString = FString::Join(SupportedAttributesStringList, TEXT(", "));
+
+			UE_LOG(LogNNERuntimeORT, Log, TEXT("%d: %s (%s)"), i, *(GetAdapterName(CurrentAdapter.Get())), *SupportedAttributesString);
+		}
+
+		if (!Adapter && (!bForceComputeOnlyDevice || !bIsGraphicsAdapter))
+		{
+			Adapter = std::move(CurrentAdapter);
+		}
+	}
+
+	if (!Adapter)
+	{
+		UE_LOG(LogNNERuntimeORT, Log, TEXT("No NPU adapter found!"));
+
+		return {};
+	}
+
+	if (bVerbose)
+	{
+		UE_LOG(LogNNERuntimeORT, Log, TEXT("Selecting NPU adapter: %s"), *(GetAdapterName(Adapter.Get())));
+	}
+
+	return Adapter;
+}
+#endif // PLATFORM_WINDOWS
+
+bool IsD3D12DeviceNPUAvailable()
+{
+#if PLATFORM_WINDOWS
+	using Microsoft::WRL::ComPtr;
+
+	ComPtr<IUnknown> AdapterNpu = GetAdapterNpu(true);
+	if (!AdapterNpu)
+	{
+		return false;
+	}
+
+	ComPtr<ID3D12Device1> D3D12Device = CreateD3D12Device(AdapterNpu.Get());
+	if (!D3D12Device)
+	{
+		return false;
+	}
+
+	return true;
+#else
+	return false;
+#endif // PLATFORM_WINDOWS
+}
+
 // For more details about ORT graph optimization checkout
 // https://onnxruntime.ai/docs/performance/model-optimizations/graph-optimizations.html
 
@@ -94,11 +293,13 @@ static constexpr FGraphOptimizationLevels OrtCpuOptimizationLevels
 // note: optimize with DirectML EP enabled, but currently an offline optimized model can not be optimized again (only DML)!
 // Therefore, if one enables offline optimization, set it to ORT_ENABLE_ALL and disable any optimization in online mode (ORT_DISABLE_ALL).
 //
-// note: since cooked models contain only basic graph optimizations, we need full optimization in online mode.
-// Therefore, offline optimization in non-Editor can not be turned on.
+// note: since during cooking the DirectML Execution Provider might not be available, one can not optimize at all, because with Float16 
+// Cast operators would be inserted, since the optimizer prepares the model for execution on CPU (at the moment this can not be turned off)
+//
+// Therefore we only optimize online for now!
 static constexpr FGraphOptimizationLevels OrtDmlOptimizationLevels
 {
-	.Cooking = GraphOptimizationLevel::ORT_ENABLE_BASIC,
+	.Cooking = GraphOptimizationLevel::ORT_DISABLE_ALL,
 	.Offline = GraphOptimizationLevel::ORT_DISABLE_ALL,
 	.Online = GraphOptimizationLevel::ORT_ENABLE_ALL
 };
@@ -197,6 +398,8 @@ TUniquePtr<Ort::SessionOptions> CreateSessionOptionsDefault(const TSharedRef<FEn
 TUniquePtr<Ort::SessionOptions> CreateSessionOptionsForDirectML(const TSharedRef<FEnvironment> &Environment, bool bRHID3D12Required)
 {
 #if PLATFORM_WINDOWS
+	using Microsoft::WRL::ComPtr;
+
 	TUniquePtr<Ort::SessionOptions> SessionOptions = CreateSessionOptionsDefault(Environment);
 	if (!SessionOptions.IsValid())
 	{
@@ -237,7 +440,7 @@ TUniquePtr<Ort::SessionOptions> CreateSessionOptionsForDirectML(const TSharedRef
 
 	// In order to use DirectML we need D3D12
 	ID3D12DynamicRHI* RHI = nullptr;
-	if (IsRHID3D12() )
+	if (IsRHID3D12())
 	{
 		RHI = GetID3D12DynamicRHI();
 	}
@@ -258,8 +461,8 @@ TUniquePtr<Ort::SessionOptions> CreateSessionOptionsForDirectML(const TSharedRef
 	check(RHI);
 
 	const int32 DeviceIndex = 0;
-	ID3D12Device* D3D12Device = RHI->RHIGetDevice(DeviceIndex);
 
+	ID3D12Device* D3D12Device = RHI->RHIGetDevice(DeviceIndex);
 	if (!D3D12Device)
 	{
 		UE_LOG(LogNNERuntimeORT, Error, TEXT("Failed to get D3D12 Device from RHI for device index %d"), DeviceIndex);
@@ -274,28 +477,101 @@ TUniquePtr<Ort::SessionOptions> CreateSessionOptionsForDirectML(const TSharedRef
 		DmlCreateFlags |= DML_CREATE_DEVICE_FLAG_DEBUG;
 	}
 
-	IDMLDevice* DmlDevice = nullptr;
-	HRESULT Res = DMLCreateDevice(D3D12Device, DmlCreateFlags, DML_PPV_ARGS(&DmlDevice));
+	ComPtr<IDMLDevice> DmlDevice;
+	HRESULT Hr = DMLCreateDevice(D3D12Device, DmlCreateFlags, IID_PPV_ARGS(&DmlDevice));
+	if (FAILED(Hr))
+	{
+		UE_LOG(LogNNERuntimeORT, Error, TEXT("Failed to create DirectML device, DMLCreateDevice error code :%x"), Hr);
+		return {};
+	}
+
+	const OrtDmlApi* DmlApi = nullptr;
+	Ort::ThrowOnError(Ort::GetApi().GetExecutionProviderApi("DML", ORT_API_VERSION, reinterpret_cast<const void**>(&DmlApi)));
+	if (!DmlApi)
+	{
+		UE_LOG(LogNNERuntimeORT, Error, TEXT("Ort DirectML Api not available!"));
+		return {};
+	}
+
+	OrtStatusPtr Status = DmlApi->SessionOptionsAppendExecutionProvider_DML1(*SessionOptions, DmlDevice.Get(), RHI->RHIGetCommandQueue());
+	if (Status)
+	{
+		UE_LOG(LogNNERuntimeORT, Error, TEXT("Failed to add DirectML execution provider to OnnxRuntime session options: %s"), ANSI_TO_TCHAR(Ort::GetApi().GetErrorMessage(Status)));
+		return {};
+	}
+
+	return SessionOptions;
+#else
+	return {};
+#endif // PLATFORM_WINDOWS
+}
 
-	if (FAILED(Res) || !DmlDevice)
+TUniquePtr<Ort::SessionOptions> CreateSessionOptionsForDirectMLNpu(const TSharedRef<FEnvironment>& Environment)
+{
+#if PLATFORM_WINDOWS
+	using Microsoft::WRL::ComPtr;
+
+	TUniquePtr<Ort::SessionOptions> SessionOptions = CreateSessionOptionsDefault(Environment);
+	if (!SessionOptions.IsValid())
 	{
-		UE_LOG(LogNNERuntimeORT, Error, TEXT("Failed to create DirectML device, DMLCreateDevice error code :%x"), Res);
 		return {};
 	}
 
-	ID3D12CommandQueue* CmdQ = RHI->RHIGetCommandQueue();
+	// Configure for DirectML
+	SessionOptions->SetExecutionMode(ExecutionMode::ORT_SEQUENTIAL);
+	SessionOptions->DisableMemPattern();
+
+	ComPtr<IUnknown> AdapterNpu = GetAdapterNpu(false);
+	if (!AdapterNpu)
+	{
+		// GetAdapterNpu() will already report any errors occured during adapter creation and...
+		// ...if there is no NPU we don't want to report an error.
+		return {};
+	}
+
+	ComPtr<ID3D12Device1> D3D12DeviceNpu = CreateD3D12Device(AdapterNpu.Get());
+	if (!D3D12DeviceNpu)
+	{
+		// CreateD3D12Device() will report any error.
+		return {};
+	}
+
+	DML_CREATE_DEVICE_FLAGS DmlCreateFlags = DML_CREATE_DEVICE_FLAG_NONE;
+
+	// Set debugging flags
+	if (GRHIGlobals.IsDebugLayerEnabled)
+	{
+		DmlCreateFlags |= DML_CREATE_DEVICE_FLAG_DEBUG;
+	}
+
+	ComPtr<IDMLDevice> DmlDeviceNpu;
+	HRESULT Hr = DMLCreateDevice(D3D12DeviceNpu.Get(), DmlCreateFlags, IID_PPV_ARGS(DmlDeviceNpu.GetAddressOf()));
+	if (FAILED(Hr))
+	{
+		UE_LOG(LogNNERuntimeORT, Error, TEXT("Failed to create DirectML device, DMLCreateDevice error code :%x"), Hr);
+		return {};
+	}
+
+	D3D12_COMMAND_QUEUE_DESC CommandQueueDesc{};
+	CommandQueueDesc.Type = D3D12_COMMAND_LIST_TYPE::D3D12_COMMAND_LIST_TYPE_COMPUTE;
+
+	ComPtr<ID3D12CommandQueue> CommandQueue;
+	Hr = D3D12DeviceNpu->CreateCommandQueue(&CommandQueueDesc, IID_PPV_ARGS(CommandQueue.GetAddressOf()));
+	if (FAILED(Hr))
+	{
+		UE_LOG(LogNNERuntimeORT, Error, TEXT("Failed to create DirectML Command Queue, CreateCommandQueue error code :%x"), Hr);
+		return {};
+	}
 
 	const OrtDmlApi* DmlApi = nullptr;
 	Ort::ThrowOnError(Ort::GetApi().GetExecutionProviderApi("DML", ORT_API_VERSION, reinterpret_cast<const void**>(&DmlApi)));
-
 	if (!DmlApi)
 	{
 		UE_LOG(LogNNERuntimeORT, Error, TEXT("Ort DirectML Api not available!"));
 		return {};
 	}
 
-	OrtStatusPtr Status = DmlApi->SessionOptionsAppendExecutionProvider_DML1(*SessionOptions, DmlDevice, CmdQ);
-
+	OrtStatusPtr Status = DmlApi->SessionOptionsAppendExecutionProvider_DML1(*SessionOptions, DmlDeviceNpu.Get(), CommandQueue.Get());
 	if (Status)
 	{
 		UE_LOG(LogNNERuntimeORT, Error, TEXT("Failed to add DirectML execution provider to OnnxRuntime session options: %s"), ANSI_TO_TCHAR(Ort::GetApi().GetErrorMessage(Status)));
diff --git a/Engine/Plugins/NNE/NNERuntimeORT/Source/NNERuntimeORT/Private/NNERuntimeORTUtils.h b/Engine/Plugins/NNE/NNERuntimeORT/Source/NNERuntimeORT/Private/NNERuntimeORTUtils.h
--- a/Engine/Plugins/NNE/NNERuntimeORT/Source/NNERuntimeORT/Private/NNERuntimeORTUtils.h
+++ b/Engine/Plugins/NNE/NNERuntimeORT/Source/NNERuntimeORT/Private/NNERuntimeORTUtils.h
@@ -15,6 +15,8 @@ namespace UE::NNERuntimeORT::Private
 
 	bool IsD3D12Available();
 
+	bool IsD3D12DeviceNPUAvailable();
+
 namespace OrtHelper
 {
 	TArray<uint32> GetShape(const Ort::Value& OrtTensor);
@@ -28,6 +30,8 @@ namespace OrtHelper
 
 	TUniquePtr<Ort::SessionOptions> CreateSessionOptionsForDirectML(const TSharedRef<FEnvironment> &Environment, bool bRHID3D12Required = true);
 
+	TUniquePtr<Ort::SessionOptions> CreateSessionOptionsForDirectMLNpu(const TSharedRef<FEnvironment>& Environment);
+
 	bool OptimizeModel(const TSharedRef<FEnvironment>& Environment, Ort::SessionOptions& SessionOptions, 
 		TConstArrayView64<uint8>& InputModel, TArray64<uint8>& OptimizedModel);
 
