[NNE] NNERuntimeORT relax error reporting.

#tests successfully ran Content Worker (Code Only) preflight with shelved change

#rnx

diff --git a/Engine/Plugins/NNE/NNERuntimeORT/Source/NNERuntimeORT/Private/NNERuntimeORTUtils.cpp b/Engine/Plugins/NNE/NNERuntimeORT/Source/NNERuntimeORT/Private/NNERuntimeORTUtils.cpp
--- a/Engine/Plugins/NNE/NNERuntimeORT/Source/NNERuntimeORT/Private/NNERuntimeORTUtils.cpp
+++ b/Engine/Plugins/NNE/NNERuntimeORT/Source/NNERuntimeORT/Private/NNERuntimeORTUtils.cpp
@@ -37,14 +37,14 @@ Microsoft::WRL::ComPtr<ID3D12Device1> CreateD3D12Device(IUnknown* AdapterPtr)
 	void* D3D12Module = FPlatformProcess::GetDllHandle(TEXT("d3d12.dll"));
 	if (!D3D12Module)
 	{
-		UE_LOG(LogNNERuntimeORT, Error, TEXT("Failed to load module 'd3d12.dll'"));
+		UE_LOG(LogNNERuntimeORT, Log, TEXT("Failed to load module 'd3d12.dll'"));
 		return {};
 	}
 
 	decltype(&D3D12CreateDevice) D3D12CreateDeviceFun = reinterpret_cast<decltype(&D3D12CreateDevice)>(FPlatformProcess::GetDllExport(D3D12Module, TEXT("D3D12CreateDevice")));
 	if (!D3D12CreateDeviceFun)
 	{
-		UE_LOG(LogNNERuntimeORT, Error, TEXT("Failed to get export 'D3D12CreateDevice' from module 'd3d12.dll'"));
+		UE_LOG(LogNNERuntimeORT, Log, TEXT("Failed to get export 'D3D12CreateDevice' from module 'd3d12.dll'"));
 		return {};
 	}
 
@@ -52,7 +52,7 @@ Microsoft::WRL::ComPtr<ID3D12Device1> CreateD3D12Device(IUnknown* AdapterPtr)
 	HRESULT Hr = D3D12CreateDeviceFun(AdapterPtr, D3D_FEATURE_LEVEL_1_0_CORE, IID_PPV_ARGS(&D3D12Device));
 	if (FAILED(Hr))
 	{
-		UE_LOG(LogNNERuntimeORT, Error, TEXT("Failed to create D3D12 device, D3D12CreateDevice error code :%x"), Hr);
+		UE_LOG(LogNNERuntimeORT, Log, TEXT("Failed to create D3D12 device, D3D12CreateDevice error code :%x"), Hr);
 		return {};
 	}
 
@@ -102,7 +102,7 @@ bool IsD3D12Available()
 	const uint32 AdapterCount = AdapterList->GetAdapterCount();
 	if (AdapterCount <= DeviceIndex)
 	{
-		UE_LOG(LogNNERuntimeORT, Error, TEXT("Invalid device index %d. Number of available devices is %d."), DeviceIndex, AdapterCount);
+		UE_LOG(LogNNERuntimeORT, Warning, TEXT("Invalid device index %d. Number of available devices is %d."), DeviceIndex, AdapterCount);
 		return false;
 	}
 
@@ -172,7 +172,7 @@ Microsoft::WRL::ComPtr<IUnknown> GetAdapterNpu(bool bVerbose)
 	DXCoreCreateAdapterFactoryFunType* DxCoreCreateAdapterFactoryFun = reinterpret_cast<DXCoreCreateAdapterFactoryFunType*>(FPlatformProcess::GetDllExport(DxCoreModule, TEXT("DXCoreCreateAdapterFactory")));
 	if (!DxCoreCreateAdapterFactoryFun)
 	{
-		UE_LOG(LogNNERuntimeORT, Error, TEXT("Failed to get export 'DXCoreCreateAdapterFactory' from module 'DXCore.dll'"));
+		UE_LOG(LogNNERuntimeORT, Log, TEXT("Failed to get export 'DXCoreCreateAdapterFactory' from module 'DXCore.dll'"));
 		return {};
 	}
 
@@ -180,7 +180,7 @@ Microsoft::WRL::ComPtr<IUnknown> GetAdapterNpu(bool bVerbose)
 	HRESULT Hr = DxCoreCreateAdapterFactoryFun(IID_PPV_ARGS(&Factory));
 	if (FAILED(Hr))
 	{
-		UE_LOG(LogNNERuntimeORT, Error, TEXT("Failed to create DXCore Adapter Factory, DxCoreCreateAdapterFactory error code :%x"), Hr);
+		UE_LOG(LogNNERuntimeORT, Log, TEXT("Failed to create DXCore Adapter Factory, DxCoreCreateAdapterFactory error code :%x"), Hr);
 		return {};
 	}
 
@@ -190,7 +190,7 @@ Microsoft::WRL::ComPtr<IUnknown> GetAdapterNpu(bool bVerbose)
 	Hr = Factory->CreateAdapterList(ARRAYSIZE(DxGUIDs), DxGUIDs, IID_PPV_ARGS(&AdapterList));
 	if (FAILED(Hr))
 	{
-		UE_LOG(LogNNERuntimeORT, Error, TEXT("Failed to create DXCore Adapter List, IDXCoreAdapterFactory::CreateAdapterList error code :%x"), Hr);
+		UE_LOG(LogNNERuntimeORT, Log, TEXT("Failed to create DXCore Adapter List, IDXCoreAdapterFactory::CreateAdapterList error code :%x"), Hr);
 		return {};
 	}
 
@@ -524,15 +524,14 @@ TUniquePtr<Ort::SessionOptions> CreateSessionOptionsForDirectMLNpu(const TShared
 	ComPtr<IUnknown> AdapterNpu = GetAdapterNpu(false);
 	if (!AdapterNpu)
 	{
-		// GetAdapterNpu() will already report any errors occured during adapter creation and...
-		// ...if there is no NPU we don't want to report an error.
+		UE_LOG(LogNNERuntimeORT, Error, TEXT("Failed to get NPU adapter."));
 		return {};
 	}
 
 	ComPtr<ID3D12Device1> D3D12DeviceNpu = CreateD3D12Device(AdapterNpu.Get());
 	if (!D3D12DeviceNpu)
 	{
-		// CreateD3D12Device() will report any error.
+		UE_LOG(LogNNERuntimeORT, Error, TEXT("Failed to create NPU device."));
 		return {};
 	}
 
