[NNE] NNE internal cleanup step2

Problem: NNE plugin are dependant on Source/NNE/Internal
Solution (step2): remove NNEBaseModel.h and implement in ORT and RDGHlsl plugin instead.

Test:
Win64 editor debug : NNEQA all tests
Win64 package debug : NNEQA all tests

#rnx

diff --git a/Engine/Plugins/Experimental/NNERuntimeRDG/Source/NNERuntimeRDG/Private/NNERuntimeRDGModel.cpp b/Engine/Plugins/Experimental/NNERuntimeRDG/Source/NNERuntimeRDG/Private/NNERuntimeRDGModel.cpp
--- a/Engine/Plugins/Experimental/NNERuntimeRDG/Source/NNERuntimeRDG/Private/NNERuntimeRDGModel.cpp
+++ b/Engine/Plugins/Experimental/NNERuntimeRDG/Source/NNERuntimeRDG/Private/NNERuntimeRDGModel.cpp
@@ -36,8 +36,6 @@ bool FModelInstanceRDG::LoadModel(TConstArrayView<uint8> ModelData, FNNERuntimeF
 	OperatorInputTensorIndices.Empty();
 	OperatorOutputTensorIndices.Empty();
 
-	TensorIdxSpan = Format.Tensors.Num();
-
 	// Add tensors
 	for (int32 Idx = 0; Idx < Format.Tensors.Num(); ++Idx)
 	{
@@ -108,17 +106,47 @@ bool FModelInstanceRDG::LoadModel(TConstArrayView<uint8> ModelData, FNNERuntimeF
 	return true;
 }
 
+TConstArrayView<NNE::FTensorDesc> FModelInstanceRDG::GetInputTensorDescs() const
+{
+	return InputSymbolicTensors;
+}
+
+TConstArrayView<NNE::FTensorDesc> FModelInstanceRDG::GetOutputTensorDescs() const
+{
+	return OutputSymbolicTensors;
+}
+
+TConstArrayView<NNE::FTensorShape> FModelInstanceRDG::GetInputTensorShapes() const
+{
+	return InputTensorShapes;
+}
 
+TConstArrayView<NNE::FTensorShape> FModelInstanceRDG::GetOutputTensorShapes() const
+{
+	return OutputTensorShapes;
+}
 
 FModelInstanceRDG::ESetInputTensorShapesStatus FModelInstanceRDG::SetInputTensorShapes(TConstArrayView<NNE::FTensorShape> InInputShapes)
 {
 	OutputTensorShapes.Reset(OutputTensorIndices.Num());
 
 	//Verify input shape are valid for the model and set InputTensorShapes
-	if (ESetInputTensorShapesStatus Status = FModelInstanceBase<NNE::IModelInstanceRDG>::SetInputTensorShapes(InInputShapes); Status != ESetInputTensorShapesStatus::Ok)
+	InputTensorShapes.Reset(InInputShapes.Num());
+	if (InInputShapes.Num() != InputSymbolicTensors.Num())
+	{
+		UE_LOG(LogNNERuntimeRDGHlsl, Error, TEXT("Number of input shapes does not match number of input tensors"));
+		return ESetInputTensorShapesStatus::Fail;
+	}
+	for (int32 i = 0; i < InInputShapes.Num(); ++i)
 	{
-		return Status;
+		const NNE::FTensorDesc SymbolicDesc = InputSymbolicTensors[i];
+		if (!InInputShapes[i].IsCompatibleWith(SymbolicDesc.GetShape()))
+		{
+			UE_LOG(LogNNERuntimeRDGHlsl, Error, TEXT("Input shape does not match input tensor %s of index %d"), *SymbolicDesc.GetName(), i);
+			return ESetInputTensorShapesStatus::Fail;
+		}
 	}
+	InputTensorShapes = InInputShapes;
 
 	//Allocate and prime all AllTensorRDGRefs with concrete shapes defaulting variables dimension to 1 if needed
 	AllTensorRDGRefs.Reserve(AllSymbolicTensorDescs.Num());
diff --git a/Engine/Plugins/Experimental/NNERuntimeRDG/Source/NNERuntimeRDG/Private/NNERuntimeRDGModel.h b/Engine/Plugins/Experimental/NNERuntimeRDG/Source/NNERuntimeRDG/Private/NNERuntimeRDGModel.h
--- a/Engine/Plugins/Experimental/NNERuntimeRDG/Source/NNERuntimeRDG/Private/NNERuntimeRDGModel.h
+++ b/Engine/Plugins/Experimental/NNERuntimeRDG/Source/NNERuntimeRDG/Private/NNERuntimeRDGModel.h
@@ -3,7 +3,6 @@
 #pragma once
 
 #include "CoreMinimal.h"
-#include "NNEModelBase.h"
 #include "NNERuntimeRDG.h"
 #include "NNETensor.h"
 #include "NNETypes.h"
@@ -11,15 +10,19 @@
 
 namespace UE::NNERuntimeRDG::Private
 {
-	class FModelInstanceRDG : public NNE::Internal::FModelInstanceBase<NNE::IModelInstanceRDG>
+	class FModelInstanceRDG : public NNE::IModelInstanceRDG
 	{
 	public:
-		using ESetInputTensorShapesStatus = IModelInstanceRDG::ESetInputTensorShapesStatus;
-		using EEnqueueRDGStatus = IModelInstanceRDG::EEnqueueRDGStatus;
+		using ESetInputTensorShapesStatus = NNE::IModelInstanceRDG::ESetInputTensorShapesStatus;
+		using EEnqueueRDGStatus = NNE::IModelInstanceRDG::EEnqueueRDGStatus;
 
 		FModelInstanceRDG() {};
 		virtual ~FModelInstanceRDG() = default;
 
+		virtual TConstArrayView<NNE::FTensorDesc> GetInputTensorDescs() const override;
+		virtual TConstArrayView<NNE::FTensorDesc> GetOutputTensorDescs() const override;
+		virtual TConstArrayView<NNE::FTensorShape> GetInputTensorShapes() const override;
+		virtual TConstArrayView<NNE::FTensorShape> GetOutputTensorShapes() const override;
 		virtual ESetInputTensorShapesStatus SetInputTensorShapes(TConstArrayView<NNE::FTensorShape> InInputShapes) override;
 		virtual EEnqueueRDGStatus EnqueueRDG(FRDGBuilder& RDGBuilder, TConstArrayView<NNE::FTensorBindingRDG> InInputBindings, TConstArrayView<NNE::FTensorBindingRDG> InOutputBindings) override;
 		
@@ -32,7 +35,11 @@ namespace UE::NNERuntimeRDG::Private
 		virtual bool PrepareModelRDG(FRDGBuilder& RDGBuilder) { return false; }
 		virtual void AddDispatchOps_RenderThread(FRDGBuilder& GraphBuilder) = 0;
 
-		int32 TensorIdxSpan;
+		//Tensor shapes and descriptions
+		TArray<NNE::FTensorShape>	InputTensorShapes;
+		TArray<NNE::FTensorShape>	OutputTensorShapes;
+		TArray<NNE::FTensorDesc>	InputSymbolicTensors;
+		TArray<NNE::FTensorDesc>	OutputSymbolicTensors;
 
 		//Tensor descriptor
 		TMap<int32, NNE::FTensorDesc>	AllSymbolicTensorDescs;
diff --git a/Engine/Plugins/NNE/NNERuntimeORT/Source/NNERuntimeORT/Private/NNERuntimeORTModel.cpp b/Engine/Plugins/NNE/NNERuntimeORT/Source/NNERuntimeORT/Private/NNERuntimeORTModel.cpp
--- a/Engine/Plugins/NNE/NNERuntimeORT/Source/NNERuntimeORT/Private/NNERuntimeORTModel.cpp
+++ b/Engine/Plugins/NNE/NNERuntimeORT/Source/NNERuntimeORT/Private/NNERuntimeORTModel.cpp
@@ -11,6 +11,7 @@
 #include "NNERuntimeORTSettings.h"
 #include "NNERuntimeORTUtils.h"
 #include "RenderGraphUtils.h"
+#include "Serialization/MemoryReader.h"
 
 #if PLATFORM_WINDOWS
 #include "ID3D12DynamicRHI.h"
@@ -133,15 +134,67 @@ namespace Detail
 
 } // namespace Detail
 
+template <class T>
+TConstArrayView<NNE::FTensorDesc> FModelInstanceORTBase<T>::GetInputTensorDescs() const
+{
+	return InputSymbolicTensors;
+}
+
+template <class T>
+TConstArrayView<NNE::FTensorDesc> FModelInstanceORTBase<T>::GetOutputTensorDescs() const
+{
+	return OutputSymbolicTensors;
+}
+
+template <class T>
+TConstArrayView<NNE::FTensorShape> FModelInstanceORTBase<T>::GetInputTensorShapes() const
+{
+	return InputTensorShapes;
+}
+
+template <class T>
+TConstArrayView<NNE::FTensorShape> FModelInstanceORTBase<T>::GetOutputTensorShapes() const
+{
+	return OutputTensorShapes;
+}
+
+template <class T>
+typename T::ESetInputTensorShapesStatus FModelInstanceORTBase<T>::SetInputTensorShapes(TConstArrayView<NNE::FTensorShape> InInputShapes)
+{
+	InputTensorShapes.Reset(InInputShapes.Num());
+
+	if (InInputShapes.Num() != InputSymbolicTensors.Num())
+	{
+		UE_LOG(LogNNERuntimeORT, Error, TEXT("Number of input shapes does not match number of input tensors"));
+		return T::ESetInputTensorShapesStatus::Fail;
+	}
+
+	for (int32 i = 0; i < InInputShapes.Num(); ++i)
+	{
+		const NNE::FTensorDesc SymbolicDesc = InputSymbolicTensors[i];
+		if (!InInputShapes[i].IsCompatibleWith(SymbolicDesc.GetShape()))
+		{
+			UE_LOG(LogNNERuntimeORT, Error, TEXT("Input shape does not match input tensor %s of index %d"), *SymbolicDesc.GetName(), i);
+			return T::ESetInputTensorShapesStatus::Fail;
+		}
+	}
+
+	InputTensorShapes = InInputShapes;
+
+	//Implementations are responsible to handle output and intermediate tensor shape inference.
+	//This base implementation only validate that all inputs are matching what the model can support.
+	return T::ESetInputTensorShapesStatus::Ok;
+}
+
 template <class ModelInterface, class TensorBinding> 
-FModelInstanceORTBase<ModelInterface, TensorBinding>::FModelInstanceORTBase(const FRuntimeConf& InRuntimeConf, TSharedRef<FEnvironment> InEnvironment)
+FModelInstanceORTRunSync<ModelInterface, TensorBinding>::FModelInstanceORTRunSync(const FRuntimeConf& InRuntimeConf, TSharedRef<FEnvironment> InEnvironment)
 	: RuntimeConf(InRuntimeConf), Environment(InEnvironment)
 {
 
 }
 
 template <class ModelInterface, class TensorBinding>
-FModelInstanceORTBase<ModelInterface, TensorBinding>::~FModelInstanceORTBase()
+FModelInstanceORTRunSync<ModelInterface, TensorBinding>::~FModelInstanceORTRunSync()
 {
 	Session.Reset();
 	if (!TempDirForModelWithExternalData.IsEmpty())
@@ -154,7 +207,7 @@ FModelInstanceORTBase<ModelInterface, TensorBinding>::~FModelInstanceORTBase()
 }
 
 template <class ModelInterface, class TensorBinding> 
-bool FModelInstanceORTBase<ModelInterface, TensorBinding>::Init(TConstArrayView64<uint8> ModelData)
+bool FModelInstanceORTRunSync<ModelInterface, TensorBinding>::Init(TConstArrayView64<uint8> ModelData)
 {
 #if WITH_EDITOR
 	try
@@ -200,7 +253,7 @@ bool FModelInstanceORTBase<ModelInterface, TensorBinding>::Init(TConstArrayView6
 }
 
 template <class ModelInterface, class TensorBinding> 
-bool FModelInstanceORTBase<ModelInterface, TensorBinding>::InitializedAndConfigureMembers()
+bool FModelInstanceORTRunSync<ModelInterface, TensorBinding>::InitializedAndConfigureMembers()
 {
 	Allocator = MakeUnique<Ort::AllocatorWithDefaultOptions>();
 	MemoryInfo = MakeUnique<Ort::MemoryInfo>(Ort::MemoryInfo::CreateCpu(OrtDeviceAllocator, OrtMemTypeCPU));
@@ -209,10 +262,10 @@ bool FModelInstanceORTBase<ModelInterface, TensorBinding>::InitializedAndConfigu
 }
 
 template <class ModelInterface, class TensorBinding>
-bool FModelInstanceORTBase<ModelInterface, TensorBinding>::ConfigureTensors(bool bAreTensorInputs)
+bool FModelInstanceORTRunSync<ModelInterface, TensorBinding>::ConfigureTensors(bool bAreTensorInputs)
 {
 	const uint32 NumberTensors							= bAreTensorInputs ? Session->GetInputCount() : Session->GetOutputCount();
-	TArray<NNE::FTensorDesc>& SymbolicTensorDescs		= bAreTensorInputs ? NNE::Internal::FModelInstanceBase<ModelInterface>::InputSymbolicTensors : NNE::Internal::FModelInstanceBase<ModelInterface>::OutputSymbolicTensors;
+	TArray<NNE::FTensorDesc>& SymbolicTensorDescs		= bAreTensorInputs ? FModelInstanceORTBase<ModelInterface>::InputSymbolicTensors : FModelInstanceORTBase<ModelInterface>::OutputSymbolicTensors;
 	TArray<ONNXTensorElementDataType>& TensorsORTType	= bAreTensorInputs ? InputTensorsORTType	: OutputTensorsORTType;
 	TArray<char*>& TensorNames							= bAreTensorInputs ? InputTensorNames		: OutputTensorNames;
 	TArray<Ort::AllocatedStringPtr>& TensorNameValues	= bAreTensorInputs ? InputTensorNameValues	: OutputTensorNameValues;
@@ -253,9 +306,9 @@ bool FModelInstanceORTBase<ModelInterface, TensorBinding>::ConfigureTensors(bool
 }
 
 template <class ModelInterface, class TensorBinding> 
-typename ModelInterface::ESetInputTensorShapesStatus FModelInstanceORTBase<ModelInterface, TensorBinding>::SetInputTensorShapes(TConstArrayView<NNE::FTensorShape> InInputShapes)
+typename ModelInterface::ESetInputTensorShapesStatus FModelInstanceORTRunSync<ModelInterface, TensorBinding>::SetInputTensorShapes(TConstArrayView<NNE::FTensorShape> InInputShapes)
 {
-	using ModelInstanceBase = NNE::Internal::FModelInstanceBase<ModelInterface>;
+	using ModelInstanceBase = FModelInstanceORTBase<ModelInterface>;
 
 	InputTensors.Reset();
 	OutputTensors.Reset();
@@ -309,9 +362,9 @@ Ort::Value CreateTensor(const Ort::MemoryInfo& MemoryInfo, const TensorBinding&
 }
 
 template <class ModelInterface, class TensorBinding>
-typename ModelInterface::ERunSyncStatus FModelInstanceORTBase<ModelInterface, TensorBinding>::RunSync(TConstArrayView<TensorBinding> InInputBindings, TConstArrayView<TensorBinding> InOutputBindings)
+typename ModelInterface::ERunSyncStatus FModelInstanceORTRunSync<ModelInterface, TensorBinding>::RunSync(TConstArrayView<TensorBinding> InInputBindings, TConstArrayView<TensorBinding> InOutputBindings)
 {
-	SCOPED_NAMED_EVENT_TEXT("FModelInstanceORTBase::RunSync", FColor::Magenta);
+	SCOPED_NAMED_EVENT_TEXT("FModelInstanceORTRunSync::RunSync", FColor::Magenta);
 
 	if (!Session.IsValid())
 	{
@@ -320,15 +373,15 @@ typename ModelInterface::ERunSyncStatus FModelInstanceORTBase<ModelInterface, Te
 	}
 
 	// Verify the model inputs were prepared
-	if (NNE::Internal::FModelInstanceBase<ModelInterface>::InputTensorShapes.IsEmpty())
+	if (FModelInstanceORTBase<ModelInterface>::InputTensorShapes.IsEmpty())
 	{
 		UE_LOG(LogNNERuntimeORT, Error, TEXT("Input shapes are not set, please call SetInputTensorShapes."));
 		return ModelInterface::ERunSyncStatus::Fail;
 	}
 
-	check(NNE::Internal::FModelInstanceBase<ModelInterface>::InputTensorShapes.Num() == InputTensors.Num());
-	check(NNE::Internal::FModelInstanceBase<ModelInterface>::InputTensorShapes.Num() == InputTensorNames.Num());
-	check(NNE::Internal::FModelInstanceBase<ModelInterface>::InputSymbolicTensors.Num() == InputTensors.Num());
+	check(FModelInstanceORTBase<ModelInterface>::InputTensorShapes.Num() == InputTensors.Num());
+	check(FModelInstanceORTBase<ModelInterface>::InputTensorShapes.Num() == InputTensorNames.Num());
+	check(FModelInstanceORTBase<ModelInterface>::InputSymbolicTensors.Num() == InputTensors.Num());
 
 	if (InInputBindings.Num() != InputTensors.Num())
 	{
@@ -336,7 +389,7 @@ typename ModelInterface::ERunSyncStatus FModelInstanceORTBase<ModelInterface, Te
 		return ModelInterface::ERunSyncStatus::Fail;
 	}
 
-	check(NNE::Internal::FModelInstanceBase<ModelInterface>::OutputSymbolicTensors.Num() == OutputTensorNames.Num());
+	check(FModelInstanceORTBase<ModelInterface>::OutputSymbolicTensors.Num() == OutputTensorNames.Num());
 
 	if (!InOutputBindings.IsEmpty() && InOutputBindings.Num() != OutputTensorNames.Num())
 	{
@@ -392,16 +445,16 @@ typename ModelInterface::ERunSyncStatus FModelInstanceORTBase<ModelInterface, Te
 		// At this (latest) stage shapes are known, therefore set them if not present yet and possibly copy data to output binding
 		if (OutputTensors.IsEmpty())
 		{
-			check(NNE::Internal::FModelInstanceBase<ModelInterface>::OutputTensorShapes.IsEmpty());
+			check(FModelInstanceORTBase<ModelInterface>::OutputTensorShapes.IsEmpty());
 
 			for (int32 i = 0; i < OutputTensorNames.Num(); i++)
 			{
-				const NNE::FTensorDesc& TensorDesc = NNE::Internal::FModelInstanceBase<ModelInterface>::OutputSymbolicTensors[i];
+				const NNE::FTensorDesc& TensorDesc = FModelInstanceORTBase<ModelInterface>::OutputSymbolicTensors[i];
 				const NNE::FTensorShape Shape = NNE::FTensorShape::Make(OrtHelper::GetShape(OrtOutputTensors[i]));
 				const NNE::Internal::FTensor Tensor = NNE::Internal::FTensor::Make(TensorDesc.GetName(), Shape, TensorDesc.GetDataType());
 
 				OutputTensors.Add(Tensor);
-				NNE::Internal::FModelInstanceBase<ModelInterface>::OutputTensorShapes.Add(Shape);
+				FModelInstanceORTBase<ModelInterface>::OutputTensorShapes.Add(Shape);
 
 				if (!InOutputBindings.IsEmpty() &&
 					InOutputBindings[i].Data &&
@@ -449,7 +502,7 @@ FModelORTCpu::FModelORTCpu(TSharedRef<FEnvironment> InEnvironment, TSharedRef<UE
 
 bool FModelInstanceORTCpu::InitializedAndConfigureMembers()
 {
-	if (!FModelInstanceORTBase::InitializedAndConfigureMembers())
+	if (!FModelInstanceORTRunSync::InitializedAndConfigureMembers())
 	{
 		return false;
 	}
@@ -488,7 +541,7 @@ FModelORTDmlGPU::FModelORTDmlGPU(TSharedRef<FEnvironment> InEnvironment, TShared
 
 bool FModelInstanceORTDmlGPU::InitializedAndConfigureMembers()
 {
-	if (!FModelInstanceORTBase::InitializedAndConfigureMembers())
+	if (!FModelInstanceORTRunSync::InitializedAndConfigureMembers())
 	{
 		return false;
 	}
@@ -670,7 +723,7 @@ FModelInstanceORTDmlRDG::ESetInputTensorShapesStatus FModelInstanceORTDmlRDG::Se
 	OutputTensorShapes.Reset();
 
 	// Verify input shape are valid for the model and set InputTensorShapes
-	if (ESetInputTensorShapesStatus Status = NNE::Internal::FModelInstanceBase<NNE::IModelInstanceRDG>::SetInputTensorShapes(InInputShapes); Status != ESetInputTensorShapesStatus::Ok)
+	if (ESetInputTensorShapesStatus Status = FModelInstanceORTBase<NNE::IModelInstanceRDG>::SetInputTensorShapes(InInputShapes); Status != ESetInputTensorShapesStatus::Ok)
 	{
 		return Status;
 	}
@@ -1000,7 +1053,7 @@ FModelORTNpu::FModelORTNpu(TSharedRef<FEnvironment> InEnvironment, TSharedRef<UE
 
 bool FModelInstanceORTNpu::InitializedAndConfigureMembers()
 {
-	if (!FModelInstanceORTBase::InitializedAndConfigureMembers())
+	if (!FModelInstanceORTRunSync::InitializedAndConfigureMembers())
 	{
 		return false;
 	}
diff --git a/Engine/Plugins/NNE/NNERuntimeORT/Source/NNERuntimeORT/Private/NNERuntimeORTModel.h b/Engine/Plugins/NNE/NNERuntimeORT/Source/NNERuntimeORT/Private/NNERuntimeORTModel.h
--- a/Engine/Plugins/NNE/NNERuntimeORT/Source/NNERuntimeORT/Private/NNERuntimeORTModel.h
+++ b/Engine/Plugins/NNE/NNERuntimeORT/Source/NNERuntimeORT/Private/NNERuntimeORTModel.h
@@ -3,7 +3,6 @@
 #pragma once
 
 #include "CoreMinimal.h"
-#include "NNEModelBase.h"
 #include "NNEModelData.h"
 #include "NNEOnnxruntime.h"
 #include "NNERuntimeCPU.h"
@@ -22,13 +21,36 @@ struct FRuntimeConf
 	ExecutionMode ExecutionMode = ExecutionMode::ORT_SEQUENTIAL;
 };
 
+template <class ModelInterface>
+class FModelInstanceORTBase : public ModelInterface
+{
+public:
+
+	virtual ~FModelInstanceORTBase() = default;
+
+	virtual TConstArrayView<NNE::FTensorDesc> GetInputTensorDescs() const override;
+	virtual TConstArrayView<NNE::FTensorDesc> GetOutputTensorDescs() const override;
+	virtual TConstArrayView<NNE::FTensorShape> GetInputTensorShapes() const override;
+	virtual TConstArrayView<NNE::FTensorShape> GetOutputTensorShapes() const override;
+	virtual typename ModelInterface::ESetInputTensorShapesStatus SetInputTensorShapes(TConstArrayView<NNE::FTensorShape> InInputShapes) override;
+
+protected:
+
+	FModelInstanceORTBase() {}
+
+	TArray<NNE::FTensorShape>	InputTensorShapes;
+	TArray<NNE::FTensorShape>	OutputTensorShapes;
+	TArray<NNE::FTensorDesc>	InputSymbolicTensors;
+	TArray<NNE::FTensorDesc>	OutputSymbolicTensors;
+};
+
 template <class ModelInterface, class TensorBinding>
-class FModelInstanceORTBase : public NNE::Internal::FModelInstanceBase<ModelInterface>
+class FModelInstanceORTRunSync : public FModelInstanceORTBase<ModelInterface>
 {
 
 public:
-	FModelInstanceORTBase(const FRuntimeConf& InRuntimeConf, TSharedRef<FEnvironment> InEnvironment);
-	virtual ~FModelInstanceORTBase();
+	FModelInstanceORTRunSync(const FRuntimeConf& InRuntimeConf, TSharedRef<FEnvironment> InEnvironment);
+	virtual ~FModelInstanceORTRunSync();
 
 	virtual typename ModelInterface::ESetInputTensorShapesStatus SetInputTensorShapes(TConstArrayView<NNE::FTensorShape> InInputShapes) override;
 
@@ -63,10 +85,10 @@ protected:
 	TArray<NNE::Internal::FTensor> OutputTensors;
 };
 
-class FModelInstanceORTCpu : public FModelInstanceORTBase<NNE::IModelInstanceCPU, NNE::FTensorBindingCPU>
+class FModelInstanceORTCpu : public FModelInstanceORTRunSync<NNE::IModelInstanceCPU, NNE::FTensorBindingCPU>
 {
 public:
-	FModelInstanceORTCpu(const FRuntimeConf& InRuntimeConf, TSharedRef<FEnvironment> InEnvironment) : FModelInstanceORTBase(InRuntimeConf, InEnvironment) {}
+	FModelInstanceORTCpu(const FRuntimeConf& InRuntimeConf, TSharedRef<FEnvironment> InEnvironment) : FModelInstanceORTRunSync(InRuntimeConf, InEnvironment) {}
 	virtual ~FModelInstanceORTCpu() = default;
 
 private:
@@ -100,10 +122,10 @@ private:
 	TSharedRef<UE::NNE::FSharedModelData> ModelData;
 };
 
-class FModelInstanceORTDmlGPU : public FModelInstanceORTBase<NNE::IModelInstanceGPU, NNE::FTensorBindingCPU>
+class FModelInstanceORTDmlGPU : public FModelInstanceORTRunSync<NNE::IModelInstanceGPU, NNE::FTensorBindingCPU>
 {
 public:
-	FModelInstanceORTDmlGPU(const FRuntimeConf& InRuntimeConf, TSharedRef<FEnvironment> InEnvironment) : FModelInstanceORTBase(InRuntimeConf, InEnvironment) {}
+	FModelInstanceORTDmlGPU(const FRuntimeConf& InRuntimeConf, TSharedRef<FEnvironment> InEnvironment) : FModelInstanceORTRunSync(InRuntimeConf, InEnvironment) {}
 	virtual ~FModelInstanceORTDmlGPU() = default;
 
 private:
@@ -124,7 +146,7 @@ private:
 	TSharedRef<UE::NNE::FSharedModelData> ModelData;
 };
 
-class FModelInstanceORTDmlRDG : public NNE::Internal::FModelInstanceBase<NNE::IModelInstanceRDG>
+class FModelInstanceORTDmlRDG : public FModelInstanceORTBase<NNE::IModelInstanceRDG>
 {
 
 public:
@@ -170,10 +192,10 @@ protected:
 	TArray<NNE::Internal::FTensor> OutputTensors;
 };
 
-class FModelInstanceORTNpu : public FModelInstanceORTBase<NNE::IModelInstanceNPU, NNE::FTensorBindingCPU>
+class FModelInstanceORTNpu : public FModelInstanceORTRunSync<NNE::IModelInstanceNPU, NNE::FTensorBindingCPU>
 {
 public:
-	FModelInstanceORTNpu(const FRuntimeConf& InRuntimeConf, TSharedRef<FEnvironment> InEnvironment) : FModelInstanceORTBase(InRuntimeConf, InEnvironment) {}
+	FModelInstanceORTNpu(const FRuntimeConf& InRuntimeConf, TSharedRef<FEnvironment> InEnvironment) : FModelInstanceORTRunSync(InRuntimeConf, InEnvironment) {}
 	virtual ~FModelInstanceORTNpu() = default;
 
 private:
diff --git a/Engine/Plugins/NNE/NNERuntimeORT/Source/NNERuntimeORT/Private/NNERuntimeORTUtils.h b/Engine/Plugins/NNE/NNERuntimeORT/Source/NNERuntimeORT/Private/NNERuntimeORTUtils.h
--- a/Engine/Plugins/NNE/NNERuntimeORT/Source/NNERuntimeORT/Private/NNERuntimeORTUtils.h
+++ b/Engine/Plugins/NNE/NNERuntimeORT/Source/NNERuntimeORT/Private/NNERuntimeORTUtils.h
@@ -4,7 +4,6 @@
 
 #include "Containers/Array.h"
 #include "NNEOnnxruntime.h"
-#include "NNERuntimeFormat.h"
 #include "NNETypes.h"
 #include "Templates/UniquePtr.h"
 
diff --git a/Engine/Source/Runtime/NNE/Internal/NNEModelBase.h b/Engine/Source/Runtime/NNE/Internal/NNEModelBase.h
deleted file mode 100644
