[NNE] NNERuntimeORT defer Environment creation

#tests Editor & Game, Win64 run models with ORT

#rnx

diff --git a/Engine/Plugins/NNE/NNERuntimeORT/Source/NNERuntimeORT/Private/NNERuntimeORTEnv.cpp b/Engine/Plugins/NNE/NNERuntimeORT/Source/NNERuntimeORT/Private/NNERuntimeORTEnv.cpp
--- a/Engine/Plugins/NNE/NNERuntimeORT/Source/NNERuntimeORT/Private/NNERuntimeORTEnv.cpp
+++ b/Engine/Plugins/NNE/NNERuntimeORT/Source/NNERuntimeORT/Private/NNERuntimeORTEnv.cpp
@@ -1,23 +1,47 @@
 // Copyright Epic Games, Inc. All Rights Reserved.
 
 #include "NNERuntimeORTEnv.h"
+#include "Misc/ScopeLock.h"
 
 namespace UE::NNERuntimeORT::Private
 {
 
-void FEnvironment::CreateOrtEnv(const FConfig& InConfig)
+void FEnvironment::Configure(const FConfig& InConfig)
 {
+	FScopeLock ScopeLock(&CriticalSection);
+
 #if WITH_EDITOR
 	// For reloading need to make sure OrtApi::ReleaseEnv() is called before we create a new Ort::Env!
 	OrtEnvironment = Ort::Env(nullptr);
 #endif // WITH_EDITOR
+
 	checkf(OrtEnvironment == Ort::Env(nullptr), TEXT("Ort::Env already created!"));
 
-	if (InConfig.bUseGlobalThreadPool)
+	Config = InConfig;
+}
+
+const Ort::Env& FEnvironment::GetOrtEnv() const
+{
+	FScopeLock ScopeLock(&CriticalSection);
+	
+	// We defer ORT Environment creation to first usage
+	if (OrtEnvironment == Ort::Env(nullptr))
+	{
+		CreateOrtEnv();
+	}
+
+	return OrtEnvironment;
+}
+
+void FEnvironment::CreateOrtEnv() const
+{
+	checkf(OrtEnvironment == Ort::Env(nullptr), TEXT("Ort::Env already created!"));
+
+	if (Config.bUseGlobalThreadPool)
 	{
 		Ort::ThreadingOptions ThreadingOptions;
-		ThreadingOptions.SetGlobalIntraOpNumThreads(InConfig.IntraOpNumThreads);
-		ThreadingOptions.SetGlobalInterOpNumThreads(InConfig.InterOpNumThreads);
+		ThreadingOptions.SetGlobalIntraOpNumThreads(Config.IntraOpNumThreads);
+		ThreadingOptions.SetGlobalInterOpNumThreads(Config.InterOpNumThreads);
 
 		// Calls OrtApi::CreateEnvWithGlobalThreadPools(), needs to be called in conjunction
 		// with OrtApi::DisablePerSessionThreads or else the session will use its own thread pools.
@@ -33,13 +57,6 @@ void FEnvironment::CreateOrtEnv(const FConfig& InConfig)
 	// WITH_EDITOR After this point, if we created a new Ort::Env, all Onnxruntime structures (ideally) need to be recreated!
 	// At least SessionOption and Session...
 	// 
-
-	Config = InConfig;
-}
-
-const Ort::Env& FEnvironment::GetOrtEnv() const
-{
-	return OrtEnvironment;
 }
 
 } // UE::NNERuntimeORT::Private
\ No newline at end of file
diff --git a/Engine/Plugins/NNE/NNERuntimeORT/Source/NNERuntimeORT/Private/NNERuntimeORTEnv.h b/Engine/Plugins/NNE/NNERuntimeORT/Source/NNERuntimeORT/Private/NNERuntimeORTEnv.h
--- a/Engine/Plugins/NNE/NNERuntimeORT/Source/NNERuntimeORT/Private/NNERuntimeORTEnv.h
+++ b/Engine/Plugins/NNE/NNERuntimeORT/Source/NNERuntimeORT/Private/NNERuntimeORTEnv.h
@@ -3,6 +3,7 @@
 #pragma once
 
 #include "NNEOnnxruntime.h"
+#include "HAL/CriticalSection.h"
 
 namespace UE::NNERuntimeORT::Private
 {
@@ -21,15 +22,19 @@ public:
 	
 	~FEnvironment() = default;
 
-	void CreateOrtEnv(const FConfig& InConfig);
+	void Configure(const FConfig& InConfig);
 
 	const Ort::Env& GetOrtEnv() const;
 
 	FConfig GetConfig() const { return Config; }
 
 private:
+	void CreateOrtEnv() const;
+
 	FConfig Config{};
-	Ort::Env OrtEnvironment{nullptr};
+	mutable Ort::Env OrtEnvironment{nullptr};
+
+	mutable FCriticalSection CriticalSection;
 };
 
 } // UE::NNERuntimeORT::Private
\ No newline at end of file
diff --git a/Engine/Plugins/NNE/NNERuntimeORT/Source/NNERuntimeORT/Private/NNERuntimeORTModule.cpp b/Engine/Plugins/NNE/NNERuntimeORT/Source/NNERuntimeORT/Private/NNERuntimeORTModule.cpp
--- a/Engine/Plugins/NNE/NNERuntimeORT/Source/NNERuntimeORT/Private/NNERuntimeORTModule.cpp
+++ b/Engine/Plugins/NNE/NNERuntimeORT/Source/NNERuntimeORT/Private/NNERuntimeORTModule.cpp
@@ -56,7 +56,7 @@ namespace UE::NNERuntimeORT::Private
 			Config.IntraOpNumThreads = ThreadingOptions.IntraOpNumThreads;
 			Config.InterOpNumThreads = ThreadingOptions.InterOpNumThreads;
 
-			Environment.CreateOrtEnv(Config);
+			Environment.Configure(Config);
 		}
 	} // EnvironmentHelper
 } // namespace UE::NNERuntimeRDG::Private
