[NNE] Internal Cleanup - step 5.3: remove dependancy from NNERuntimeORT and NNERuntimeIREE to NNE/Internal/NNETensor.h

Problem: NNERuntimeORT and NNERuntimeIREE are relying on the internal of /NNE
Solution: Brake the dependancy by having each plugin get it's minimal implementaion of a concrete tensor.

Test:
Win64 editor debug and dev: NNEQA all tests + styletransfer
Win64 package dev : NNEQA all tests + styletransfer

#rnx

diff --git a/Engine/Plugins/Experimental/NNERuntimeIREE/Source/NNERuntimeIREE/Private/NNERuntimeIREEModel.cpp b/Engine/Plugins/Experimental/NNERuntimeIREE/Source/NNERuntimeIREE/Private/NNERuntimeIREEModel.cpp
--- a/Engine/Plugins/Experimental/NNERuntimeIREE/Source/NNERuntimeIREE/Private/NNERuntimeIREEModel.cpp
+++ b/Engine/Plugins/Experimental/NNERuntimeIREE/Source/NNERuntimeIREE/Private/NNERuntimeIREEModel.cpp
@@ -14,8 +14,8 @@
 #include "NNE.h"
 #include "NNEModelData.h"
 #include "NNERuntimeIREELog.h"
+#include "NNERuntimeIREETensor.h"
 #include "NNEStatus.h"
-#include "NNETensor.h"
 #include "RenderGraphUtils.h"
 #include "HAL/UnrealMemory.h"
 #include "Serialization/Archive.h"
@@ -1283,7 +1283,7 @@ namespace UE::NNERuntimeIREE
 					{
 						for (int32 i = 0; i < OutputTensorDescs.Num(); i++)
 						{
-							NNE::Internal::FTensor Tensor = NNE::Internal::FTensor::MakeFromSymbolicDesc(OutputTensorDescs[i]);
+							NNERuntimeIREE::Private::FTensor Tensor = NNERuntimeIREE::Private::FTensor::MakeFromSymbolicDesc(OutputTensorDescs[i]);
 							OutputTensors.Add(Tensor);
 							OutputTensorShapes.Add(Tensor.GetShape());
 						}
@@ -1305,7 +1305,7 @@ namespace UE::NNERuntimeIREE
 							return ESetInputTensorShapesStatus::Fail;
 						}
 
-						NNE::Internal::FTensor Tensor = NNE::Internal::FTensor::Make(SymbolicDesc.GetName(), InInputShapes[i], SymbolicDesc.GetDataType());
+						NNERuntimeIREE::Private::FTensor Tensor = NNERuntimeIREE::Private::FTensor::Make(InInputShapes[i], SymbolicDesc.GetDataType());
 						InputTensors.Emplace(Tensor);
 					}
 					InputTensorShapes = InInputShapes;
@@ -1497,8 +1497,8 @@ namespace UE::NNERuntimeIREE
 				TArray<UE::NNE::FTensorDesc> OutputTensorDescs;
 				TArray<UE::NNE::FTensorShape> InputTensorShapes;
 				TArray<UE::NNE::FTensorShape> OutputTensorShapes;
-				TArray<NNE::Internal::FTensor> InputTensors;
-				TArray<NNE::Internal::FTensor> OutputTensors;
+				TArray<NNERuntimeIREE::Private::FTensor> InputTensors;
+				TArray<NNERuntimeIREE::Private::FTensor> OutputTensors;
 
 				FSession(TSharedRef<FDevice> InDevice, TSharedRef<UE::NNERuntimeIREE::Private::FModule> InModule, iree_runtime_session_t* InSession, const iree_runtime_call_t& InCall, TConstArrayView<UE::NNE::FTensorDesc> InInputTensorDescs, TConstArrayView<UE::NNE::FTensorDesc> InOutputTensorDescs)
 					: Device(InDevice), Module(InModule), Session(InSession), Call(InCall), InputTensorDescs(InInputTensorDescs), OutputTensorDescs(InOutputTensorDescs)
diff --git a/Engine/Plugins/Experimental/NNERuntimeIREE/Source/NNERuntimeIREE/Private/NNERuntimeIREETensor.h b/Engine/Plugins/Experimental/NNERuntimeIREE/Source/NNERuntimeIREE/Private/NNERuntimeIREETensor.h
new file mode 100644
--- /dev/null
+++ b/Engine/Plugins/Experimental/NNERuntimeIREE/Source/NNERuntimeIREE/Private/NNERuntimeIREETensor.h
@@ -0,0 +1,44 @@
+// Copyright Epic Games, Inc. All Rights Reserved.
+
+#pragma once
+
+#include "NNETypes.h"
+
+namespace UE::NNERuntimeIREE::Private
+{
+	class FTensor
+	{
+	private:
+		NNE::FTensorShape	Shape;
+		uint64				DataSize = 0;
+
+		FTensor() = default;
+
+	public:
+		const NNE::FTensorShape& GetShape() const
+		{
+			return Shape;
+		}
+
+		uint64 GetDataSize() const
+		{
+			return DataSize;
+		}
+
+		static FTensor Make(const NNE::FTensorShape& InShape, ENNETensorDataType DataType)
+		{
+			check(InShape.Volume() <= TNumericLimits<uint32>::Max());
+
+			FTensor Tensor;
+			Tensor.Shape = InShape;
+			Tensor.DataSize = (uint64)NNE::GetTensorDataTypeSizeInBytes(DataType) * InShape.Volume();
+			return Tensor;
+		}
+
+		static FTensor MakeFromSymbolicDesc(const NNE::FTensorDesc& TensorDesc)
+		{
+			return Make(NNE::FTensorShape::MakeFromSymbolic(TensorDesc.GetShape()), TensorDesc.GetDataType());
+		}
+	};
+
+} // namespace UE::NNERuntimeIREE::Private
\ No newline at end of file
diff --git a/Engine/Plugins/NNE/NNERuntimeORT/Source/NNERuntimeORT/Private/NNERuntimeORTModel.cpp b/Engine/Plugins/NNE/NNERuntimeORT/Source/NNERuntimeORT/Private/NNERuntimeORTModel.cpp
--- a/Engine/Plugins/NNE/NNERuntimeORT/Source/NNERuntimeORT/Private/NNERuntimeORTModel.cpp
+++ b/Engine/Plugins/NNE/NNERuntimeORT/Source/NNERuntimeORT/Private/NNERuntimeORTModel.cpp
@@ -323,7 +323,7 @@ typename ModelInterface::ESetInputTensorShapesStatus FModelInstanceORTRunSync<Mo
 	// Setup concrete input tensor
 	for (int32 i = 0; i < ModelInstanceBase::InputSymbolicTensors.Num(); ++i)
 	{
-		NNE::Internal::FTensor Tensor = NNE::Internal::FTensor::Make(ModelInstanceBase::InputSymbolicTensors[i].GetName(), InInputShapes[i], ModelInstanceBase::InputSymbolicTensors[i].GetDataType());
+		FTensor Tensor = FTensor::Make(InInputShapes[i], ModelInstanceBase::InputSymbolicTensors[i].GetDataType());
 		InputTensors.Emplace(Tensor);
 	}
 
@@ -332,7 +332,7 @@ typename ModelInterface::ESetInputTensorShapesStatus FModelInstanceORTRunSync<Mo
 	{
 		if (SymbolicTensorDesc.GetShape().IsConcrete())
 		{
-			NNE::Internal::FTensor Tensor = NNE::Internal::FTensor::MakeFromSymbolicDesc(SymbolicTensorDesc);
+			FTensor Tensor = FTensor::MakeFromSymbolicDesc(SymbolicTensorDesc);
 			OutputTensors.Emplace(Tensor);
 			ModelInstanceBase::OutputTensorShapes.Emplace(Tensor.GetShape());
 		}
@@ -347,7 +347,7 @@ typename ModelInterface::ESetInputTensorShapesStatus FModelInstanceORTRunSync<Mo
 }
 
 template <class TensorBinding>
-Ort::Value CreateTensor(const Ort::MemoryInfo& MemoryInfo, const TensorBinding& Binding, const NNE::Internal::FTensor& Tensor, const ONNXTensorElementDataType ElementDataType)
+Ort::Value CreateTensor(const Ort::MemoryInfo& MemoryInfo, const TensorBinding& Binding, const FTensor& Tensor, const ONNXTensorElementDataType ElementDataType)
 {
 	const uint64 SizeInBytes = Tensor.GetDataSize();
 	const uint32 ShapeLen = (uint32)Tensor.GetShape().Rank();
@@ -405,7 +405,7 @@ typename ModelInterface::ERunSyncStatus FModelInstanceORTRunSync<ModelInterface,
 		for (int32 i = 0; i < InputTensorNames.Num(); i++)
 		{
 			const TensorBinding& Binding = InInputBindings[i];
-			const NNE::Internal::FTensor& Tensor = InputTensors[i];
+			const FTensor& Tensor = InputTensors[i];
 
 			if (!Binding.Data && Binding.SizeInBytes != 0)
 			{
@@ -451,7 +451,7 @@ typename ModelInterface::ERunSyncStatus FModelInstanceORTRunSync<ModelInterface,
 			{
 				const NNE::FTensorDesc& TensorDesc = FModelInstanceORTBase<ModelInterface>::OutputSymbolicTensors[i];
 				const NNE::FTensorShape Shape = NNE::FTensorShape::Make(OrtHelper::GetShape(OrtOutputTensors[i]));
-				const NNE::Internal::FTensor Tensor = NNE::Internal::FTensor::Make(TensorDesc.GetName(), Shape, TensorDesc.GetDataType());
+				const FTensor Tensor = FTensor::Make(Shape, TensorDesc.GetDataType());
 
 				OutputTensors.Add(Tensor);
 				FModelInstanceORTBase<ModelInterface>::OutputTensorShapes.Add(Shape);
@@ -744,7 +744,7 @@ FModelInstanceORTDmlRDG::ESetInputTensorShapesStatus FModelInstanceORTDmlRDG::Se
 	{
 		for (int32 i = 0; i < InputSymbolicTensors.Num(); i++)
 		{
-			NNE::Internal::FTensor Tensor = NNE::Internal::FTensor::Make(InputSymbolicTensors[i].GetName(), InInputShapes[i], InputSymbolicTensors[i].GetDataType());
+			FTensor Tensor = FTensor::Make(InInputShapes[i], InputSymbolicTensors[i].GetDataType());
 			InputTensors.Emplace(Tensor);
 		}
 
@@ -755,7 +755,7 @@ FModelInstanceORTDmlRDG::ESetInputTensorShapesStatus FModelInstanceORTDmlRDG::Se
 
 			if (SymbolicTensorDesc.GetShape().IsConcrete())
 			{
-				NNE::Internal::FTensor Tensor = NNE::Internal::FTensor::MakeFromSymbolicDesc(SymbolicTensorDesc);
+				FTensor Tensor = FTensor::MakeFromSymbolicDesc(SymbolicTensorDesc);
 				OutputTensors.Emplace(Tensor);
 				OutputTensorShapes.Emplace(Tensor.GetShape());
 			}
@@ -782,7 +782,7 @@ FModelInstanceORTDmlRDG::ESetInputTensorShapesStatus FModelInstanceORTDmlRDG::Se
 	// Setup concrete input tensors
 	for (int32 i = 0; i < InputSymbolicTensors.Num(); i++)
 	{
-		NNE::Internal::FTensor Tensor = NNE::Internal::FTensor::Make(InputSymbolicTensors[i].GetName(), InInputShapes[i], InputSymbolicTensors[i].GetDataType());
+		FTensor Tensor = FTensor::Make(InInputShapes[i], InputSymbolicTensors[i].GetDataType());
 		InputTensors.Emplace(Tensor);
 
 		const NNE::FSymbolicTensorShape& SymbolicInputShape = InputSymbolicTensors[i].GetShape();
@@ -825,7 +825,7 @@ FModelInstanceORTDmlRDG::ESetInputTensorShapesStatus FModelInstanceORTDmlRDG::Se
 
 		if (SymbolicTensorDesc.GetShape().IsConcrete())
 		{
-			NNE::Internal::FTensor Tensor = NNE::Internal::FTensor::MakeFromSymbolicDesc(SymbolicTensorDesc);
+			FTensor Tensor = FTensor::MakeFromSymbolicDesc(SymbolicTensorDesc);
 			OutputTensors.Emplace(Tensor);
 			OutputTensorShapes.Emplace(Tensor.GetShape());
 		}
@@ -847,7 +847,7 @@ FModelInstanceORTDmlRDG::ESetInputTensorShapesStatus FModelInstanceORTDmlRDG::Se
 	return ESetInputTensorShapesStatus::Ok;
 }
 
-Ort::Value CreateTensor(const OrtDmlApi& DmlApi, const Ort::MemoryInfo& MemoryInfo, FRHIBuffer* Buffer, const NNE::Internal::FTensor& Tensor, ONNXTensorElementDataType ElementDataType,
+Ort::Value CreateTensor(const OrtDmlApi& DmlApi, const Ort::MemoryInfo& MemoryInfo, FRHIBuffer* Buffer, const FTensor& Tensor, ONNXTensorElementDataType ElementDataType,
 	TArray<std::unique_ptr<void, void (*)(void*)>>& DmlAllocatorResources)
 {
 	checkf(Buffer, TEXT("CreateTensor needs Buffer to be set"));
diff --git a/Engine/Plugins/NNE/NNERuntimeORT/Source/NNERuntimeORT/Private/NNERuntimeORTModel.h b/Engine/Plugins/NNE/NNERuntimeORT/Source/NNERuntimeORT/Private/NNERuntimeORTModel.h
--- a/Engine/Plugins/NNE/NNERuntimeORT/Source/NNERuntimeORT/Private/NNERuntimeORTModel.h
+++ b/Engine/Plugins/NNE/NNERuntimeORT/Source/NNERuntimeORT/Private/NNERuntimeORTModel.h
@@ -9,8 +9,8 @@
 #include "NNERuntimeNPU.h"
 #include "NNERuntimeGPU.h"
 #include "NNERuntimeORTEnv.h"
+#include "NNERuntimeORTTensor.h"
 #include "NNERuntimeRDG.h"
-#include "NNETensor.h"
 #include "NNETypes.h"
 
 namespace UE::NNERuntimeORT::Private
@@ -81,8 +81,8 @@ protected:
 	TArray<char*> InputTensorNames;
 	TArray<char*> OutputTensorNames;
 
-	TArray<NNE::Internal::FTensor> InputTensors;
-	TArray<NNE::Internal::FTensor> OutputTensors;
+	TArray<FTensor> InputTensors;
+	TArray<FTensor> OutputTensors;
 };
 
 class FModelInstanceORTCpu : public FModelInstanceORTRunSync<NNE::IModelInstanceCPU, NNE::FTensorBindingCPU>
@@ -188,8 +188,8 @@ protected:
 	TArray<TArray<FString>> InputSymbolicDimensionNames;
 	TArray<TArray<FString>> OutputSymbolicDimensionNames;
 
-	TArray<NNE::Internal::FTensor> InputTensors;
-	TArray<NNE::Internal::FTensor> OutputTensors;
+	TArray<FTensor> InputTensors;
+	TArray<FTensor> OutputTensors;
 };
 
 class FModelInstanceORTNpu : public FModelInstanceORTRunSync<NNE::IModelInstanceNPU, NNE::FTensorBindingCPU>
diff --git a/Engine/Plugins/NNE/NNERuntimeORT/Source/NNERuntimeORT/Private/NNERuntimeORTTensor.h b/Engine/Plugins/NNE/NNERuntimeORT/Source/NNERuntimeORT/Private/NNERuntimeORTTensor.h
new file mode 100644
--- /dev/null
+++ b/Engine/Plugins/NNE/NNERuntimeORT/Source/NNERuntimeORT/Private/NNERuntimeORTTensor.h
@@ -0,0 +1,44 @@
+// Copyright Epic Games, Inc. All Rights Reserved.
+
+#pragma once
+
+#include "NNETypes.h"
+
+namespace UE::NNERuntimeORT::Private
+{
+	class FTensor
+	{
+	private:
+		NNE::FTensorShape	Shape;
+		uint64				DataSize = 0;
+
+		FTensor() = default;
+
+	public:
+		const NNE::FTensorShape& GetShape() const
+		{
+			return Shape;
+		}
+
+		uint64 GetDataSize() const
+		{
+			return DataSize;
+		}
+
+		static FTensor Make(const NNE::FTensorShape& InShape, ENNETensorDataType DataType)
+		{
+			check(InShape.Volume() <= TNumericLimits<uint32>::Max());
+
+			FTensor Tensor;
+			Tensor.Shape = InShape;
+			Tensor.DataSize = (uint64)NNE::GetTensorDataTypeSizeInBytes(DataType) * InShape.Volume();
+			return Tensor;
+		}
+
+		static FTensor MakeFromSymbolicDesc(const NNE::FTensorDesc& TensorDesc)
+		{
+			return Make(NNE::FTensorShape::MakeFromSymbolic(TensorDesc.GetShape()), TensorDesc.GetDataType());
+		}
+	};
+
+} // namespace UE::NNERuntimeORT::Private
\ No newline at end of file
