[NNE] NNERuntimeORTDml fix Adapter enumeration for NPU.

Problem:
Starting with Windows 11 Version 24H2, the NPU driver reports back that it is a  DXCORE_ADAPTER_ATTRIBUTE_D3D12_GENERIC_ML device rather than DXCORE_ADAPTER_ATTRIBUTE_D3D12_CORE_COMPUTE.

Solution:
Add an additional GetAdapterList call with the adapter attribute GUID value from DXCORE_ADAPTER_ATTRIBUTE_D3D12_GENERIC_ML.

#tests Win64, Game, verified models run on NPU

#rnx

diff --git a/Engine/Plugins/NNE/NNERuntimeORT/Source/NNERuntimeORT/NNERuntimeORT.Build.cs b/Engine/Plugins/NNE/NNERuntimeORT/Source/NNERuntimeORT/NNERuntimeORT.Build.cs
--- a/Engine/Plugins/NNE/NNERuntimeORT/Source/NNERuntimeORT/NNERuntimeORT.Build.cs
+++ b/Engine/Plugins/NNE/NNERuntimeORT/Source/NNERuntimeORT/NNERuntimeORT.Build.cs
@@ -1,6 +1,7 @@
 // Copyright Epic Games, Inc. All Rights Reserved.
 
 using UnrealBuildTool;
+using System;
 
 public class NNERuntimeORT : ModuleRules
 {
@@ -32,6 +33,11 @@ public class NNERuntimeORT : ModuleRules
 				"DirectML",
 				"DX12"
 			});
+
+			if (Version.TryParse(Target.WindowsPlatform.WindowsSdkVersion, out Version WindowsSdkVersion) && WindowsSdkVersion.Build >= 26100 /* Windows 11, version 24H2 */)
+			{
+				PublicDefinitions.Add("ORT_USE_NEW_DXCORE_FEATURES");
+			}
 		}
 
 		if (Target.bBuildEditor)
diff --git a/Engine/Plugins/NNE/NNERuntimeORT/Source/NNERuntimeORT/Private/NNERuntimeORTUtils.cpp b/Engine/Plugins/NNE/NNERuntimeORT/Source/NNERuntimeORT/Private/NNERuntimeORTUtils.cpp
--- a/Engine/Plugins/NNE/NNERuntimeORT/Source/NNERuntimeORT/Private/NNERuntimeORTUtils.cpp
+++ b/Engine/Plugins/NNE/NNERuntimeORT/Source/NNERuntimeORT/Private/NNERuntimeORTUtils.cpp
@@ -184,6 +184,56 @@ Microsoft::WRL::ComPtr<IUnknown> GetAdapterNpu(bool bVerbose)
 		return {};
 	}
 
+	// New path from Windows 11, version 24H2 on
+	{
+#ifdef ORT_USE_NEW_DXCORE_FEATURES
+		const GUID DxGUIDs[] = { DXCORE_ADAPTER_ATTRIBUTE_D3D12_GENERIC_ML };
+#else
+		const GUID DxGUIDs[] = { {0xb71b0d41, 0x1088, 0x422f, 0xa2, 0x7c, 0x2, 0x50, 0xb7, 0xd3, 0xa9, 0x88} }; // Value of DXCORE_ADAPTER_ATTRIBUTE_D3D12_GENERIC_ML
+#endif
+
+		ComPtr<IDXCoreAdapterList> AdapterList;
+		Hr = Factory->CreateAdapterList(ARRAYSIZE(DxGUIDs), DxGUIDs, IID_PPV_ARGS(&AdapterList));
+		if (FAILED(Hr))
+		{
+			UE_LOG(LogNNERuntimeORT, Log, TEXT("Failed to create DXCore Adapter List, IDXCoreAdapterFactory::CreateAdapterList error code :%x"), Hr);
+			return {};
+		}
+
+		if (bVerbose)
+		{
+			UE_LOG(LogNNERuntimeORT, Log, TEXT("Available graphics and compute adapters:"));
+		}
+
+		ComPtr<IDXCoreAdapter> Adapter;
+		for (uint32_t i = 0, AdapterCount = AdapterList->GetAdapterCount(); i < AdapterCount; i++)
+		{
+			Hr = AdapterList->GetAdapter(static_cast<uint32_t>(i), IID_PPV_ARGS(&Adapter));
+			if (FAILED(Hr))
+			{
+				UE_LOG(LogNNERuntimeORT, Warning, TEXT("%d: Failed to get adapter, IDXCoreAdapterList::GetAdapter error code: %x"), i, Hr);
+			}
+			else
+			{
+				break;
+			}
+		}
+
+		if (!Adapter)
+		{
+			UE_LOG(LogNNERuntimeORT, Log, TEXT("No NPU adapter found with attribute DXCORE_ADAPTER_ATTRIBUTE_D3D12_GENERIC_ML (Windows 11 Version 24H2 or newer)!"));
+		}
+		else
+		{
+			if (bVerbose)
+			{
+				UE_LOG(LogNNERuntimeORT, Log, TEXT("Selecting NPU adapter: %s"), *(GetAdapterName(Adapter.Get())));
+			}
+
+			return Adapter;
+		}
+	}
+
 	const GUID DxGUIDs[] = { DXCORE_ADAPTER_ATTRIBUTE_D3D12_CORE_COMPUTE };
 
 	ComPtr<IDXCoreAdapterList> AdapterList;
