[NNE] NNERuntimeORTDml fix compiler warning/error in CL 41229077.

#tests Win64, Editor build locally with clang

#rnx

diff --git a/Engine/Plugins/NNE/NNERuntimeORT/Source/NNERuntimeORT/Private/NNERuntimeORTUtils.cpp b/Engine/Plugins/NNE/NNERuntimeORT/Source/NNERuntimeORT/Private/NNERuntimeORTUtils.cpp
--- a/Engine/Plugins/NNE/NNERuntimeORT/Source/NNERuntimeORT/Private/NNERuntimeORTUtils.cpp
+++ b/Engine/Plugins/NNE/NNERuntimeORT/Source/NNERuntimeORT/Private/NNERuntimeORTUtils.cpp
@@ -189,7 +189,7 @@ Microsoft::WRL::ComPtr<IUnknown> GetAdapterNpu(bool bVerbose)
 #ifdef ORT_USE_NEW_DXCORE_FEATURES
 		const GUID DxGUIDs[] = { DXCORE_ADAPTER_ATTRIBUTE_D3D12_GENERIC_ML };
 #else
-		const GUID DxGUIDs[] = { {0xb71b0d41, 0x1088, 0x422f, 0xa2, 0x7c, 0x2, 0x50, 0xb7, 0xd3, 0xa9, 0x88} }; // Value of DXCORE_ADAPTER_ATTRIBUTE_D3D12_GENERIC_ML
+		const GUID DxGUIDs[] = { { 0xb71b0d41, 0x1088, 0x422f, { 0xa2, 0x7c, 0x2, 0x50, 0xb7, 0xd3, 0xa9, 0x88 } } }; // Value of DXCORE_ADAPTER_ATTRIBUTE_D3D12_GENERIC_ML
 #endif
 
 		ComPtr<IDXCoreAdapterList> AdapterList;
