[NNE] NNERuntimeORT runtime availability in non-Editor apps.

#tests Win64, Mac & Linux, Editor & Game, ran model tests

#rnx

diff --git a/Engine/Plugins/NNE/NNERuntimeORT/Source/NNERuntimeORT/Private/NNERuntimeORT.cpp b/Engine/Plugins/NNE/NNERuntimeORT/Source/NNERuntimeORT/Private/NNERuntimeORT.cpp
--- a/Engine/Plugins/NNE/NNERuntimeORT/Source/NNERuntimeORT/Private/NNERuntimeORT.cpp
+++ b/Engine/Plugins/NNE/NNERuntimeORT/Source/NNERuntimeORT/Private/NNERuntimeORT.cpp
@@ -22,9 +22,6 @@ DEFINE_LOG_CATEGORY(LogNNERuntimeORT);
 FGuid UNNERuntimeORTCpu::GUID = FGuid((int32)'O', (int32)'C', (int32)'P', (int32)'U');
 int32 UNNERuntimeORTCpu::Version = 0x00000004;
 
-FGuid UNNERuntimeORTDml::GUID = FGuid((int32)'O', (int32)'D', (int32)'M', (int32)'L');
-int32 UNNERuntimeORTDml::Version = 0x00000004;
-
 namespace UE::NNERuntimeORT::Private::Details
 { 
 	//Should be kept in sync with OnnxFileLoaderHelper::InitUNNEModelDataFromFile()
@@ -76,6 +73,21 @@ namespace UE::NNERuntimeORT::Private::Details
 			Writer.Serialize(const_cast<uint8*>(AdditionalFileData[OnnxExternalDataBytesKey].GetData()), AdditionalFileData[OnnxExternalDataBytesKey].Num());
 		}
 	}
+
+	bool IsAvailableGPU(bool bDirectMLAvailable, bool bD3D12Available)
+	{
+		return bDirectMLAvailable && bD3D12Available;
+	}
+
+	bool IsAvailableRDG(bool bDirectMLAvailable, bool bRHID3D12Available)
+	{
+		return bDirectMLAvailable && bRHID3D12Available;
+	}
+
+	bool IsAvailableNPU(bool bDirectMLAvailable, bool bD3D12DeviceNPUAvailable)
+	{
+		return bDirectMLAvailable && bD3D12DeviceNPUAvailable;
+	}
 } // namespace UE::NNERuntimeORT::Private::Details
 
 UNNERuntimeORTCpu::ECanCreateModelDataStatus UNNERuntimeORTCpu::CanCreateModelData(const FString& FileType, TConstArrayView64<uint8> FileData, const TMap<FString, TConstArrayView64<uint8>>& AdditionalFileData, const FGuid& FileId, const ITargetPlatform* TargetPlatform) const
@@ -183,27 +195,45 @@ TSharedPtr<UE::NNE::IModelCPU> UNNERuntimeORTCpu::CreateModelCPU(const TObjectPt
 }
 
 /*
- * UNNERuntimeORTDml
+ * FNNERuntimeORTDmlImpl
  */
-void UNNERuntimeORTDml::Init(TSharedRef<UE::NNERuntimeORT::Private::FEnvironment> InEnvironment, bool bInDirectMLAvailable)
+class FNNERuntimeORTDmlImpl : public INNERuntimeORTDmlImpl
+{
+
+	using ECanCreateModelCommonStatus = UE::NNE::EResultStatus;
+
+private:
+	TSharedPtr<UE::NNERuntimeORT::Private::FEnvironment> Environment;
+
+public:
+	static FGuid GUID;
+	static int32 Version;
+
+	FNNERuntimeORTDmlImpl() = default;
+	virtual ~FNNERuntimeORTDmlImpl() = default;
+
+virtual void Init(TSharedRef<UE::NNERuntimeORT::Private::FEnvironment> InEnvironment, bool bInDirectMLAvailable) override
 {
+	using namespace UE::NNERuntimeORT::Private;
+
 	Environment = InEnvironment;
-	bDirectMLAvailable = bInDirectMLAvailable;
-	bD3D12Available = UE::NNERuntimeORT::Private::IsD3D12Available();
-	bD3D12DeviceNPUAvailable = UE::NNERuntimeORT::Private::IsD3D12DeviceNPUAvailable();
+	
+	bIsAvailableGPU = Details::IsAvailableGPU(bInDirectMLAvailable, IsD3D12Available());
+	bIsAvailableRDG = Details::IsAvailableRDG(bInDirectMLAvailable, IsRHID3D12Available());
+	bIsAvailableNPU = Details::IsAvailableNPU(bInDirectMLAvailable, IsD3D12DeviceNPUAvailable());
 }
 
-FString UNNERuntimeORTDml::GetRuntimeName() const
+virtual FString GetRuntimeName() const override
 {
 	return TEXT("NNERuntimeORTDml");
 }
 
-UNNERuntimeORTDml::ECanCreateModelDataStatus UNNERuntimeORTDml::CanCreateModelData(const FString& FileType, TConstArrayView64<uint8> FileData, const TMap<FString, TConstArrayView64<uint8>>& AdditionalFileData, const FGuid& FileId, const ITargetPlatform* TargetPlatform) const
+virtual ECanCreateModelDataStatus CanCreateModelData(const FString& FileType, TConstArrayView64<uint8> FileData, const TMap<FString, TConstArrayView64<uint8>>& AdditionalFileData, const FGuid& FileId, const ITargetPlatform* TargetPlatform) const override
 {
 	return (!FileData.IsEmpty() && FileType.Compare("onnx", ESearchCase::IgnoreCase) == 0) ? ECanCreateModelDataStatus::Ok : ECanCreateModelDataStatus::FailFileIdNotSupported;
 }
 
-TSharedPtr<UE::NNE::FSharedModelData> UNNERuntimeORTDml::CreateModelData(const FString& FileType, TConstArrayView64<uint8> FileData, const TMap<FString, TConstArrayView64<uint8>>& AdditionalFileData, const FGuid& FileId, const ITargetPlatform* TargetPlatform)
+virtual TSharedPtr<UE::NNE::FSharedModelData> CreateModelData(const FString& FileType, TConstArrayView64<uint8> FileData, const TMap<FString, TConstArrayView64<uint8>>& AdditionalFileData, const FGuid& FileId, const ITargetPlatform* TargetPlatform) override
 {
 	using namespace UE::NNERuntimeORT::Private;
 
@@ -239,27 +269,22 @@ TSharedPtr<UE::NNE::FSharedModelData> UNNERuntimeORTDml::CreateModelData(const F
 
 	TArray64<uint8> Result;
 	FMemoryWriter64 Writer(Result, /*bIsPersitent =*/ true);
-	Writer << UNNERuntimeORTDml::GUID;
-	Writer << UNNERuntimeORTDml::Version;
+	Writer << GUID;
+	Writer << Version;
 
 	Details::WriteOnnxModelData(Writer, OptimizedModelView, AdditionalFileData);
 
 	return MakeShared<UE::NNE::FSharedModelData>(MakeSharedBufferFromArray(MoveTemp(Result)), 0);
 }
 
-FString UNNERuntimeORTDml::GetModelDataIdentifier(const FString& FileType, TConstArrayView64<uint8> FileData, const TMap<FString, TConstArrayView64<uint8>>& AdditionalFileData, const FGuid& FileId, const ITargetPlatform* TargetPlatform) const
+virtual FString GetModelDataIdentifier(const FString& FileType, TConstArrayView64<uint8> FileData, const TMap<FString, TConstArrayView64<uint8>>& AdditionalFileData, const FGuid& FileId, const ITargetPlatform* TargetPlatform) const override
 {
-	return FileId.ToString(EGuidFormats::Digits) + "-" + UNNERuntimeORTDml::GUID.ToString(EGuidFormats::Digits) + "-" + FString::FromInt(UNNERuntimeORTDml::Version);
+	return FileId.ToString(EGuidFormats::Digits) + "-" + GUID.ToString(EGuidFormats::Digits) + "-" + FString::FromInt(Version);
 }
 
-UNNERuntimeORTDml::ECanCreateModelGPUStatus UNNERuntimeORTDml::CanCreateModelGPU(const TObjectPtr<UNNEModelData> ModelData) const
+virtual ECanCreateModelGPUStatus CanCreateModelGPU(const TObjectPtr<UNNEModelData> ModelData) const override
 {
-	if (!bDirectMLAvailable)
-	{
-		return ECanCreateModelCommonStatus::Fail;
-	}
-
-	if (!bD3D12Available)
+	if (!ensureMsgf(bIsAvailableGPU, TEXT("GPU interface should not be available!")))
 	{
 		return ECanCreateModelCommonStatus::Fail;
 	}
@@ -267,7 +292,7 @@ UNNERuntimeORTDml::ECanCreateModelGPUStatus UNNERuntimeORTDml::CanCreateModelGPU
 	return CanCreateModelCommon(ModelData, false) == ECanCreateModelCommonStatus::Ok ? ECanCreateModelGPUStatus::Ok : ECanCreateModelGPUStatus::Fail;
 }
 
-TSharedPtr<UE::NNE::IModelGPU> UNNERuntimeORTDml::CreateModelGPU(const TObjectPtr<UNNEModelData> ModelData)
+virtual TSharedPtr<UE::NNE::IModelGPU> CreateModelGPU(const TObjectPtr<UNNEModelData> ModelData) override
 {
 #if PLATFORM_WINDOWS
 	check(ModelData);
@@ -286,24 +311,17 @@ TSharedPtr<UE::NNE::IModelGPU> UNNERuntimeORTDml::CreateModelGPU(const TObjectPt
 #endif // PLATFORM_WINDOWS
 }
 
-UNNERuntimeORTDml::ECanCreateModelRDGStatus UNNERuntimeORTDml::CanCreateModelRDG(TObjectPtr<UNNEModelData> ModelData) const
+virtual ECanCreateModelRDGStatus CanCreateModelRDG(TObjectPtr<UNNEModelData> ModelData) const override
 {
-	if (!bDirectMLAvailable)
-	{
-		return ECanCreateModelCommonStatus::Fail;
-	}
-
-#if PLATFORM_WINDOWS
-	if (!IsRHID3D12())
+	if (!ensureMsgf(bIsAvailableRDG, TEXT("RDG interface should not be available!")))
 	{
 		return ECanCreateModelCommonStatus::Fail;
 	}
-#endif // PLATFORM_WINDOWS
 
 	return CanCreateModelCommon(ModelData) == ECanCreateModelCommonStatus::Ok ? ECanCreateModelRDGStatus::Ok : ECanCreateModelRDGStatus::Fail;
 }
 
-TSharedPtr<UE::NNE::IModelRDG> UNNERuntimeORTDml::CreateModelRDG(TObjectPtr<UNNEModelData> ModelData)
+virtual TSharedPtr<UE::NNE::IModelRDG> CreateModelRDG(TObjectPtr<UNNEModelData> ModelData) override
 {
 #if PLATFORM_WINDOWS
 	check(ModelData);
@@ -322,14 +340,9 @@ TSharedPtr<UE::NNE::IModelRDG> UNNERuntimeORTDml::CreateModelRDG(TObjectPtr<UNNE
 #endif // PLATFORM_WINDOWS
 }
 
-UNNERuntimeORTDml::ECanCreateModelNPUStatus UNNERuntimeORTDml::CanCreateModelNPU(const TObjectPtr<UNNEModelData> ModelData) const
+virtual ECanCreateModelNPUStatus CanCreateModelNPU(const TObjectPtr<UNNEModelData> ModelData) const override
 {
-	if (!bDirectMLAvailable)
-	{
-		return ECanCreateModelCommonStatus::Fail;
-	}
-
-	if (!bD3D12DeviceNPUAvailable)
+	if (!ensureMsgf(bIsAvailableNPU, TEXT("NPU interface should not be available!")))
 	{
 		return ECanCreateModelCommonStatus::Fail;
 	}
@@ -337,7 +350,7 @@ UNNERuntimeORTDml::ECanCreateModelNPUStatus UNNERuntimeORTDml::CanCreateModelNPU
 	return CanCreateModelCommon(ModelData) == ECanCreateModelCommonStatus::Ok ? ECanCreateModelRDGStatus::Ok : ECanCreateModelRDGStatus::Fail;
 }
 
-TSharedPtr<UE::NNE::IModelNPU> UNNERuntimeORTDml::CreateModelNPU(const TObjectPtr<UNNEModelData> ModelData)
+virtual TSharedPtr<UE::NNE::IModelNPU> CreateModelNPU(const TObjectPtr<UNNEModelData> ModelData) override
 {
 #if PLATFORM_WINDOWS
 	check(ModelData);
@@ -356,13 +369,14 @@ TSharedPtr<UE::NNE::IModelNPU> UNNERuntimeORTDml::CreateModelNPU(const TObjectPt
 #endif // PLATFORM_WINDOWS
 }
 
-UNNERuntimeORTDml::ECanCreateModelCommonStatus UNNERuntimeORTDml::CanCreateModelCommon(const TObjectPtr<UNNEModelData> ModelData, bool bRHID3D12Required) const
+private:
+ECanCreateModelCommonStatus CanCreateModelCommon(const TObjectPtr<UNNEModelData> ModelData, bool bRHID3D12Required = true) const
 {
 #if PLATFORM_WINDOWS
 	check(ModelData != nullptr);
 
-	constexpr int32 GuidSize = sizeof(UNNERuntimeORTDml::GUID);
-	constexpr int32 VersionSize = sizeof(UNNERuntimeORTDml::Version);
+	constexpr int32 GuidSize = sizeof(GUID);
+	constexpr int32 VersionSize = sizeof(Version);
 	const TSharedPtr<UE::NNE::FSharedModelData> SharedData = ModelData->GetModelData(GetRuntimeName());
 
 	if (!SharedData.IsValid())
@@ -379,12 +393,76 @@ UNNERuntimeORTDml::ECanCreateModelCommonStatus UNNERuntimeORTDml::CanCreateModel
 
 	static const FGuid DeprecatedGUID = FGuid((int32)'O', (int32)'G', (int32)'P', (int32)'U');
 
-	bool bResult = FGenericPlatformMemory::Memcmp(&(Data[0]), &(UNNERuntimeORTDml::GUID), GuidSize) == 0;
+	bool bResult = FGenericPlatformMemory::Memcmp(&(Data[0]), &(GUID), GuidSize) == 0;
 	bResult |= FGenericPlatformMemory::Memcmp(&(Data[0]), &(DeprecatedGUID), GuidSize) == 0;
-	bResult &= FGenericPlatformMemory::Memcmp(&(Data[GuidSize]), &(UNNERuntimeORTDml::Version), VersionSize) == 0;
+	bResult &= FGenericPlatformMemory::Memcmp(&(Data[GuidSize]), &(Version), VersionSize) == 0;
 
 	return bResult ? ECanCreateModelCommonStatus::Ok : ECanCreateModelCommonStatus::Fail;
 #else // PLATFORM_WINDOWS
 	return ECanCreateModelCommonStatus::Fail;
 #endif // PLATFORM_WINDOWS
-}
\ No newline at end of file
+}
+
+	bool bIsAvailableGPU = false;
+	bool bIsAvailableRDG = false;
+	bool bIsAvailableNPU = false;
+};
+
+FGuid FNNERuntimeORTDmlImpl::GUID = FGuid((int32)'O', (int32)'D', (int32)'M', (int32)'L');
+int32 FNNERuntimeORTDmlImpl::Version = 0x00000004;
+
+UNNERuntimeORTDmlProxy::UNNERuntimeORTDmlProxy()
+{
+	Impl = MakeUnique<FNNERuntimeORTDmlImpl>();
+}
+
+void UNNERuntimeORTDmlProxy::Init(TSharedRef<UE::NNERuntimeORT::Private::FEnvironment> InEnvironment, bool bInDirectMLAvailable) { Impl->Init(InEnvironment, bInDirectMLAvailable); }
+
+FString UNNERuntimeORTDmlProxy::GetRuntimeName() const { return Impl->GetRuntimeName(); }
+
+UNNERuntimeORTDmlProxy::ECanCreateModelDataStatus UNNERuntimeORTDmlProxy::CanCreateModelData(const FString& FileType, TConstArrayView64<uint8> FileData, const TMap<FString, TConstArrayView64<uint8>>& AdditionalFileData, const FGuid& FileId, const ITargetPlatform* TargetPlatform) const { return Impl->CanCreateModelData(FileType, FileData, AdditionalFileData, FileId, TargetPlatform); }
+TSharedPtr<UE::NNE::FSharedModelData> UNNERuntimeORTDmlProxy::CreateModelData(const FString& FileType, TConstArrayView64<uint8> FileData, const TMap<FString, TConstArrayView64<uint8>>& AdditionalFileData, const FGuid& FileId, const ITargetPlatform* TargetPlatform) { return Impl->CreateModelData(FileType, FileData, AdditionalFileData, FileId, TargetPlatform); }
+FString UNNERuntimeORTDmlProxy::GetModelDataIdentifier(const FString& FileType, TConstArrayView64<uint8> FileData, const TMap<FString, TConstArrayView64<uint8>>& AdditionalFileData, const FGuid& FileId, const ITargetPlatform* TargetPlatform) const { return Impl->GetModelDataIdentifier(FileType, FileData, AdditionalFileData, FileId, TargetPlatform); }
+
+namespace UE::NNERuntimeORT::Private
+{
+
+TWeakObjectPtr<UNNERuntimeORTDmlProxy> MakeRuntimeDml(bool bDirectMLAvailable)
+{
+	const bool bRHID3D12Available = IsRHID3D12Available();
+	const bool bD3D12Available = IsD3D12Available();
+	const bool bD3D12DeviceNPUAvailable = IsD3D12DeviceNPUAvailable();
+
+	const bool bIsAvailableGPU = Details::IsAvailableGPU(bDirectMLAvailable, bD3D12Available);
+	const bool bIsAvailableRDG = Details::IsAvailableRDG(bDirectMLAvailable, bRHID3D12Available);
+	const bool bIsAvailableNPU = Details::IsAvailableNPU(bDirectMLAvailable, bD3D12DeviceNPUAvailable);
+
+	UE_LOG(LogNNERuntimeORT, Log, TEXT("MakeRuntimeORTDml:"));
+	UE_LOG(LogNNERuntimeORT, Log, TEXT("  DirectML:  %s"), (bDirectMLAvailable ? TEXT("yes") : TEXT("no")));
+	UE_LOG(LogNNERuntimeORT, Log, TEXT("  RHI D3D12: %s"), (bRHID3D12Available ? TEXT("yes") : TEXT("no")));
+	UE_LOG(LogNNERuntimeORT, Log, TEXT("  D3D12:     %s"), (bD3D12Available ? TEXT("yes") : TEXT("no")));
+	UE_LOG(LogNNERuntimeORT, Log, TEXT("  NPU:       %s"), (bD3D12DeviceNPUAvailable ? TEXT("yes") : TEXT("no")));
+
+	UE_LOG(LogNNERuntimeORT, Log, TEXT("Interface availability:"));
+	UE_LOG(LogNNERuntimeORT, Log, TEXT("  GPU: %s"), (bIsAvailableGPU ? TEXT("yes") : TEXT("no")));
+	UE_LOG(LogNNERuntimeORT, Log, TEXT("  RDG: %s"), (bIsAvailableRDG ? TEXT("yes") : TEXT("no")));
+	UE_LOG(LogNNERuntimeORT, Log, TEXT("  NPU: %s"), (bIsAvailableNPU ? TEXT("yes") : TEXT("no")));
+
+	if (bIsAvailableGPU && bIsAvailableRDG && bIsAvailableNPU) return NewObject<UNNERuntimeORTDml_GPU_RDG_NPU>();
+	else if (bIsAvailableGPU && bIsAvailableRDG) return NewObject<UNNERuntimeORTDml_GPU_RDG>();
+	else if (bIsAvailableGPU && bIsAvailableNPU) return NewObject<UNNERuntimeORTDml_GPU_NPU>();
+	else if (bIsAvailableRDG && bIsAvailableNPU) return NewObject<UNNERuntimeORTDml_RDG_NPU>();
+	else if (bIsAvailableGPU) return NewObject<UNNERuntimeORTDml_GPU>();
+	else if (bIsAvailableRDG) return NewObject<UNNERuntimeORTDml_RDG>();
+	else if (bIsAvailableNPU) return NewObject<UNNERuntimeORTDml_NPU>();
+
+#if WITH_EDITOR
+	UE_LOG(LogNNERuntimeORT, Log, TEXT("NNERuntimeORTDml can only cook!"));
+	return NewObject<UNNERuntimeORTDmlProxy>();
+#else
+	UE_LOG(LogNNERuntimeORT, Log, TEXT("NNERuntimeORTDml is not available!"));
+	return {};
+#endif
+}
+
+} // namespace UE::NNERuntimeORT::Private
\ No newline at end of file
diff --git a/Engine/Plugins/NNE/NNERuntimeORT/Source/NNERuntimeORT/Private/NNERuntimeORT.h b/Engine/Plugins/NNE/NNERuntimeORT/Source/NNERuntimeORT/Private/NNERuntimeORT.h
--- a/Engine/Plugins/NNE/NNERuntimeORT/Source/NNERuntimeORT/Private/NNERuntimeORT.h
+++ b/Engine/Plugins/NNE/NNERuntimeORT/Source/NNERuntimeORT/Private/NNERuntimeORT.h
@@ -46,22 +46,21 @@ public:
 	virtual TSharedPtr<UE::NNE::IModelCPU> CreateModelCPU(const TObjectPtr<UNNEModelData> ModelData) override;
 };
 
+class INNERuntimeORTDmlImpl : public INNERuntime, public INNERuntimeGPU, public INNERuntimeRDG, public INNERuntimeNPU
+{
+public:
+	virtual void Init(TSharedRef<UE::NNERuntimeORT::Private::FEnvironment> InEnvironment, bool bInDirectMLAvailable) = 0;
+};
+
 UCLASS()
-class UNNERuntimeORTDml : public UObject, public INNERuntime, public INNERuntimeGPU, public INNERuntimeRDG, public INNERuntimeNPU
+class UNNERuntimeORTDmlProxy : public UObject, public INNERuntime
 {
 	GENERATED_BODY()
 
-	using ECanCreateModelCommonStatus = UE::NNE::EResultStatus;
-
-private:
-	TSharedPtr<UE::NNERuntimeORT::Private::FEnvironment> Environment;
-
 public:
-	static FGuid GUID;
-	static int32 Version;
+	UNNERuntimeORTDmlProxy();
+	~UNNERuntimeORTDmlProxy() = default;
 	
-	virtual ~UNNERuntimeORTDml() = default;
-
 	void Init(TSharedRef<UE::NNERuntimeORT::Private::FEnvironment> InEnvironment, bool bInDirectMLAvailable);
 
 	virtual FString GetRuntimeName() const override;
@@ -70,19 +69,98 @@ public:
 	virtual TSharedPtr<UE::NNE::FSharedModelData> CreateModelData(const FString& FileType, TConstArrayView64<uint8> FileData, const TMap<FString, TConstArrayView64<uint8>>& AdditionalFileData, const FGuid& FileId, const ITargetPlatform* TargetPlatform) override;
 	virtual FString GetModelDataIdentifier(const FString& FileType, TConstArrayView64<uint8> FileData, const TMap<FString, TConstArrayView64<uint8>>& AdditionalFileData, const FGuid& FileId, const ITargetPlatform* TargetPlatform) const override;
 
-	virtual ECanCreateModelGPUStatus CanCreateModelGPU(const TObjectPtr<UNNEModelData> ModelData) const override;
-	virtual TSharedPtr<UE::NNE::IModelGPU> CreateModelGPU(const TObjectPtr<UNNEModelData> ModelData) override;
+protected:
+	TUniquePtr<INNERuntimeORTDmlImpl> Impl;
+};
 
-	virtual ECanCreateModelRDGStatus CanCreateModelRDG(TObjectPtr<UNNEModelData> ModelData) const override;
-	virtual TSharedPtr<UE::NNE::IModelRDG> CreateModelRDG(TObjectPtr<UNNEModelData> ModelData) override;
+UCLASS()
+class UNNERuntimeORTDml_GPU_RDG_NPU : public UNNERuntimeORTDmlProxy, public INNERuntimeGPU, public INNERuntimeRDG, public INNERuntimeNPU
+{
+	GENERATED_BODY()
 
-	virtual ECanCreateModelNPUStatus CanCreateModelNPU(const TObjectPtr<UNNEModelData> ModelData) const override;
-	virtual TSharedPtr<UE::NNE::IModelNPU> CreateModelNPU(const TObjectPtr<UNNEModelData> ModelData) override;
+public:
+	virtual ECanCreateModelGPUStatus CanCreateModelGPU(const TObjectPtr<UNNEModelData> ModelData) const override { return Impl->CanCreateModelGPU(ModelData); }
+	virtual TSharedPtr<UE::NNE::IModelGPU> CreateModelGPU(const TObjectPtr<UNNEModelData> ModelData) override { return Impl->CreateModelGPU(ModelData); }
 
-private:
-	ECanCreateModelCommonStatus CanCreateModelCommon(const TObjectPtr<UNNEModelData> ModelData, bool bRHID3D12Required = true) const;
+	virtual ECanCreateModelRDGStatus CanCreateModelRDG(TObjectPtr<UNNEModelData> ModelData) const override { return Impl->CanCreateModelRDG(ModelData); }
+	virtual TSharedPtr<UE::NNE::IModelRDG> CreateModelRDG(TObjectPtr<UNNEModelData> ModelData) override { return Impl->CreateModelRDG(ModelData); }
+
+	virtual ECanCreateModelNPUStatus CanCreateModelNPU(const TObjectPtr<UNNEModelData> ModelData) const override { return Impl->CanCreateModelNPU(ModelData); }
+	virtual TSharedPtr<UE::NNE::IModelNPU> CreateModelNPU(const TObjectPtr<UNNEModelData> ModelData) override { return Impl->CreateModelNPU(ModelData); }
+};
+
+UCLASS()
+class UNNERuntimeORTDml_GPU_RDG : public UNNERuntimeORTDmlProxy, public INNERuntimeGPU, public INNERuntimeRDG
+{
+	GENERATED_BODY()
+
+public:
+	virtual ECanCreateModelGPUStatus CanCreateModelGPU(const TObjectPtr<UNNEModelData> ModelData) const override { return Impl->CanCreateModelGPU(ModelData); }
+	virtual TSharedPtr<UE::NNE::IModelGPU> CreateModelGPU(const TObjectPtr<UNNEModelData> ModelData) override { return Impl->CreateModelGPU(ModelData); }
+
+	virtual ECanCreateModelRDGStatus CanCreateModelRDG(TObjectPtr<UNNEModelData> ModelData) const override { return Impl->CanCreateModelRDG(ModelData); }
+	virtual TSharedPtr<UE::NNE::IModelRDG> CreateModelRDG(TObjectPtr<UNNEModelData> ModelData) override { return Impl->CreateModelRDG(ModelData); }
+};
+
+UCLASS()
+class UNNERuntimeORTDml_GPU_NPU : public UNNERuntimeORTDmlProxy, public INNERuntimeGPU, public INNERuntimeNPU
+{
+	GENERATED_BODY()
+
+public:
+	virtual ECanCreateModelGPUStatus CanCreateModelGPU(const TObjectPtr<UNNEModelData> ModelData) const override { return Impl->CanCreateModelGPU(ModelData); }
+	virtual TSharedPtr<UE::NNE::IModelGPU> CreateModelGPU(const TObjectPtr<UNNEModelData> ModelData) override { return Impl->CreateModelGPU(ModelData); }
+
+	virtual ECanCreateModelNPUStatus CanCreateModelNPU(const TObjectPtr<UNNEModelData> ModelData) const override { return Impl->CanCreateModelNPU(ModelData); }
+	virtual TSharedPtr<UE::NNE::IModelNPU> CreateModelNPU(const TObjectPtr<UNNEModelData> ModelData) override { return Impl->CreateModelNPU(ModelData); }
+};
+
+UCLASS()
+class UNNERuntimeORTDml_RDG_NPU : public UNNERuntimeORTDmlProxy, public INNERuntimeRDG, public INNERuntimeNPU
+{
+	GENERATED_BODY()
+
+public:
+	virtual ECanCreateModelRDGStatus CanCreateModelRDG(TObjectPtr<UNNEModelData> ModelData) const override { return Impl->CanCreateModelRDG(ModelData); }
+	virtual TSharedPtr<UE::NNE::IModelRDG> CreateModelRDG(TObjectPtr<UNNEModelData> ModelData) override { return Impl->CreateModelRDG(ModelData); }
+
+	virtual ECanCreateModelNPUStatus CanCreateModelNPU(const TObjectPtr<UNNEModelData> ModelData) const override { return Impl->CanCreateModelNPU(ModelData); }
+	virtual TSharedPtr<UE::NNE::IModelNPU> CreateModelNPU(const TObjectPtr<UNNEModelData> ModelData) override { return Impl->CreateModelNPU(ModelData); }
+};
+
+UCLASS()
+class UNNERuntimeORTDml_GPU : public UNNERuntimeORTDmlProxy, public INNERuntimeGPU
+{
+	GENERATED_BODY()
+
+public:
+	virtual ECanCreateModelGPUStatus CanCreateModelGPU(const TObjectPtr<UNNEModelData> ModelData) const override { return Impl->CanCreateModelGPU(ModelData); }
+	virtual TSharedPtr<UE::NNE::IModelGPU> CreateModelGPU(const TObjectPtr<UNNEModelData> ModelData) override { return Impl->CreateModelGPU(ModelData); }
+};
+
+UCLASS()
+class UNNERuntimeORTDml_RDG : public UNNERuntimeORTDmlProxy, public INNERuntimeRDG
+{
+	GENERATED_BODY()
+
+public:
+	virtual ECanCreateModelRDGStatus CanCreateModelRDG(TObjectPtr<UNNEModelData> ModelData) const override { return Impl->CanCreateModelRDG(ModelData); }
+	virtual TSharedPtr<UE::NNE::IModelRDG> CreateModelRDG(TObjectPtr<UNNEModelData> ModelData) override { return Impl->CreateModelRDG(ModelData); }
+};
+
+UCLASS()
+class UNNERuntimeORTDml_NPU : public UNNERuntimeORTDmlProxy, public INNERuntimeNPU
+{
+	GENERATED_BODY()
+
+public:
+	virtual ECanCreateModelNPUStatus CanCreateModelNPU(const TObjectPtr<UNNEModelData> ModelData) const override { return Impl->CanCreateModelNPU(ModelData); }
+	virtual TSharedPtr<UE::NNE::IModelNPU> CreateModelNPU(const TObjectPtr<UNNEModelData> ModelData) override { return Impl->CreateModelNPU(ModelData); }
+};
+
+namespace UE::NNERuntimeORT::Private
+{
+
+TWeakObjectPtr<UNNERuntimeORTDmlProxy> MakeRuntimeDml(bool bDirectMLAvailable);
 
-	bool bDirectMLAvailable = false;
-	bool bD3D12Available = false;
-	bool bD3D12DeviceNPUAvailable = false;
-};
\ No newline at end of file
+} // namespace UE::NNERuntimeORT::Private
\ No newline at end of file
diff --git a/Engine/Plugins/NNE/NNERuntimeORT/Source/NNERuntimeORT/Private/NNERuntimeORTModule.cpp b/Engine/Plugins/NNE/NNERuntimeORT/Source/NNERuntimeORT/Private/NNERuntimeORTModule.cpp
--- a/Engine/Plugins/NNE/NNERuntimeORT/Source/NNERuntimeORT/Private/NNERuntimeORTModule.cpp
+++ b/Engine/Plugins/NNE/NNERuntimeORT/Source/NNERuntimeORT/Private/NNERuntimeORTModule.cpp
@@ -103,7 +103,7 @@ void FNNERuntimeORTModule::StartupModule()
 	EnvironmentHelper::CreateOrtEnvFromSettings(GetDefault<UNNERuntimeORTSettings>(), *Environment);
 
 	// NNE runtime ORT Dml startup
-	NNERuntimeORTDml = NewObject<UNNERuntimeORTDml>();
+	NNERuntimeORTDml = MakeRuntimeDml(bDirectMLDllLoaded);
 	if (NNERuntimeORTDml.IsValid())
 	{
 		TWeakInterfacePtr<INNERuntime> RuntimeDmlInterface(NNERuntimeORTDml.Get());
diff --git a/Engine/Plugins/NNE/NNERuntimeORT/Source/NNERuntimeORT/Private/NNERuntimeORTModule.h b/Engine/Plugins/NNE/NNERuntimeORT/Source/NNERuntimeORT/Private/NNERuntimeORTModule.h
--- a/Engine/Plugins/NNE/NNERuntimeORT/Source/NNERuntimeORT/Private/NNERuntimeORTModule.h
+++ b/Engine/Plugins/NNE/NNERuntimeORT/Source/NNERuntimeORT/Private/NNERuntimeORTModule.h
@@ -8,7 +8,7 @@
 #include "UObject/WeakObjectPtr.h"
 #include "UObject/WeakObjectPtrTemplates.h"
 
-class UNNERuntimeORTDml;
+class UNNERuntimeORTDmlProxy;
 class UNNERuntimeORTCpu;
 
 namespace UE::NNERuntimeORT::Private
@@ -19,7 +19,7 @@ namespace UE::NNERuntimeORT::Private
 class FNNERuntimeORTModule : public IModuleInterface
 {
 private:
-	TWeakObjectPtr<UNNERuntimeORTDml> NNERuntimeORTDml{ nullptr };
+	TWeakObjectPtr<UNNERuntimeORTDmlProxy> NNERuntimeORTDml{ nullptr };
 	TWeakObjectPtr<UNNERuntimeORTCpu> NNERuntimeORTCpu{ nullptr };
 
 	virtual void StartupModule() override;
diff --git a/Engine/Plugins/NNE/NNERuntimeORT/Source/NNERuntimeORT/Private/NNERuntimeORTUtils.cpp b/Engine/Plugins/NNE/NNERuntimeORT/Source/NNERuntimeORT/Private/NNERuntimeORTUtils.cpp
--- a/Engine/Plugins/NNE/NNERuntimeORT/Source/NNERuntimeORT/Private/NNERuntimeORTUtils.cpp
+++ b/Engine/Plugins/NNE/NNERuntimeORT/Source/NNERuntimeORT/Private/NNERuntimeORTUtils.cpp
@@ -60,6 +60,15 @@ Microsoft::WRL::ComPtr<ID3D12Device1> CreateD3D12Device(IUnknown* AdapterPtr)
 }
 #endif // PLATFORM_WINDOWS
 
+bool IsRHID3D12Available()
+{
+#if PLATFORM_WINDOWS
+	return IsRHID3D12();
+#else
+	return false;
+#endif // PLATFORM_WINDOWS
+}
+
 // Check for DirectX 12-compatible hardware.
 // Use DXGI to enumerate adapters and try to create a d3d12 device using the default adapter
 bool IsD3D12Available()
diff --git a/Engine/Plugins/NNE/NNERuntimeORT/Source/NNERuntimeORT/Private/NNERuntimeORTUtils.h b/Engine/Plugins/NNE/NNERuntimeORT/Source/NNERuntimeORT/Private/NNERuntimeORTUtils.h
--- a/Engine/Plugins/NNE/NNERuntimeORT/Source/NNERuntimeORT/Private/NNERuntimeORTUtils.h
+++ b/Engine/Plugins/NNE/NNERuntimeORT/Source/NNERuntimeORT/Private/NNERuntimeORTUtils.h
@@ -12,6 +12,8 @@ namespace UE::NNERuntimeORT::Private
 
 	class FEnvironment;
 
+	bool IsRHID3D12Available();
+
 	bool IsD3D12Available();
 
 	bool IsD3D12DeviceNPUAvailable();
