[NNE] NNEDenoiser fixed crash with MRQ when using tiling with overlap.

#tests Run with MRQ

#rnx

diff --git a/Engine/Plugins/NNE/NNERuntimeORT/Source/NNERuntimeORT/Private/NNERuntimeORTModel.cpp b/Engine/Plugins/NNE/NNERuntimeORT/Source/NNERuntimeORT/Private/NNERuntimeORTModel.cpp
--- a/Engine/Plugins/NNE/NNERuntimeORT/Source/NNERuntimeORT/Private/NNERuntimeORTModel.cpp
+++ b/Engine/Plugins/NNE/NNERuntimeORT/Source/NNERuntimeORT/Private/NNERuntimeORTModel.cpp
@@ -132,6 +132,33 @@ namespace Detail
 		return Session.IsValid();
 	}
 
+	bool SetInputTensorShapes(TConstArrayView<NNE::FTensorShape> InInputShapes, TConstArrayView<NNE::FTensorDesc> InputSymbolicTensors, TArray<NNE::FTensorShape> &InputTensorShapes)
+	{
+		InputTensorShapes.Reset(InInputShapes.Num());
+
+		if (InInputShapes.Num() != InputSymbolicTensors.Num())
+		{
+			UE_LOG(LogNNERuntimeORT, Error, TEXT("Number of input shapes does not match number of input tensors"));
+			return false;
+		}
+
+		for (int32 i = 0; i < InInputShapes.Num(); ++i)
+		{
+			const NNE::FTensorDesc SymbolicDesc = InputSymbolicTensors[i];
+			if (!InInputShapes[i].IsCompatibleWith(SymbolicDesc.GetShape()))
+			{
+				UE_LOG(LogNNERuntimeORT, Error, TEXT("Input shape does not match input tensor %s of index %d"), *SymbolicDesc.GetName(), i);
+				return false;
+			}
+		}
+
+		InputTensorShapes = InInputShapes;
+
+		//Implementations are responsible to handle output and intermediate tensor shape inference.
+		//This base implementation only validate that all inputs are matching what the model can support.
+		return true;
+	}
+
 } // namespace Detail
 
 template <class T>
@@ -161,29 +188,12 @@ TConstArrayView<NNE::FTensorShape> FModelInstanceORTBase<T>::GetOutputTensorShap
 template <class T>
 typename T::ESetInputTensorShapesStatus FModelInstanceORTBase<T>::SetInputTensorShapes(TConstArrayView<NNE::FTensorShape> InInputShapes)
 {
-	InputTensorShapes.Reset(InInputShapes.Num());
-
-	if (InInputShapes.Num() != InputSymbolicTensors.Num())
+	if (Detail::SetInputTensorShapes(InInputShapes, InputSymbolicTensors, InputTensorShapes))
 	{
-		UE_LOG(LogNNERuntimeORT, Error, TEXT("Number of input shapes does not match number of input tensors"));
-		return T::ESetInputTensorShapesStatus::Fail;
+		return T::ESetInputTensorShapesStatus::Ok;
 	}
 
-	for (int32 i = 0; i < InInputShapes.Num(); ++i)
-	{
-		const NNE::FTensorDesc SymbolicDesc = InputSymbolicTensors[i];
-		if (!InInputShapes[i].IsCompatibleWith(SymbolicDesc.GetShape()))
-		{
-			UE_LOG(LogNNERuntimeORT, Error, TEXT("Input shape does not match input tensor %s of index %d"), *SymbolicDesc.GetName(), i);
-			return T::ESetInputTensorShapesStatus::Fail;
-		}
-	}
-
-	InputTensorShapes = InInputShapes;
-
-	//Implementations are responsible to handle output and intermediate tensor shape inference.
-	//This base implementation only validate that all inputs are matching what the model can support.
-	return T::ESetInputTensorShapesStatus::Ok;
+	return T::ESetInputTensorShapesStatus::Fail;
 }
 
 template <class ModelInterface, class TensorBinding> 
@@ -575,20 +585,142 @@ TSharedPtr<NNE::IModelInstanceRDG> FModelORTDmlRDG::CreateModelInstanceRDG()
 	return ModelInstance;
 }
 
-FModelInstanceORTDmlRDG::FModelInstanceORTDmlRDG(TSharedRef<UE::NNE::FSharedModelData> InModelData, const FRuntimeConf& InRuntimeConf, TSharedRef<FEnvironment> InEnvironment)
-	:  ModelData(InModelData), RuntimeConf(InRuntimeConf), Environment(InEnvironment)
-{}
-
-FModelInstanceORTDmlRDG::~FModelInstanceORTDmlRDG()
+class FModelInstanceORTDmlRDG::FProxy
 {
-	Session.Reset();
-	if (!TempDirForModelWithExternalData.IsEmpty())
+public:
+	FProxy(TSharedRef<UE::NNE::FSharedModelData> InModelData, const FRuntimeConf& InRuntimeConf, TSharedRef<FEnvironment> InEnvironment)
+		: ModelData(InModelData), RuntimeConf(InRuntimeConf), Environment(InEnvironment)
+	{
+
+	}
+	
+	~FProxy()
+	{
+		Session.Reset();
+		if (!TempDirForModelWithExternalData.IsEmpty())
 	{
 		if (!IFileManager::Get().DeleteDirectory(*TempDirForModelWithExternalData, false, true))
 		{
 			UE_LOG(LogNNERuntimeORT, Warning, TEXT("Large models are an experimental feature at the moment. FModelInstanceORTDmlRDG could not delete temp directy %s on model instance destruction."), *TempDirForModelWithExternalData);
 		}
 	}
+	}
+
+	bool ConfigureTensors(bool bAreTensorInputs)
+	{
+		SCOPED_NAMED_EVENT_TEXT("FModelInstanceORTDmlRDG::ConfigureTensors", FColor::Magenta);
+
+		const uint32 NumberTensors							= bAreTensorInputs ? Session->GetInputCount()			: Session->GetOutputCount();
+		TArray<NNE::FTensorDesc>& SymbolicTensorDescs		= bAreTensorInputs ? InputSymbolicTensors		: OutputSymbolicTensors;
+		TArray<ONNXTensorElementDataType>& TensorsORTType	= bAreTensorInputs ? InputTensorsORTType			: OutputTensorsORTType;
+		TArray<char*>& TensorNames							= bAreTensorInputs ? InputTensorNames			: OutputTensorNames;
+		TArray<Ort::AllocatedStringPtr>& TensorNameValues	= bAreTensorInputs ? InputTensorNameValues		: OutputTensorNameValues;
+		TArray<TArray<FString>>& SymbolicDimensionNames 	= bAreTensorInputs ? InputSymbolicDimensionNames : OutputSymbolicDimensionNames;
+
+		SymbolicTensorDescs.Reset();
+		TensorsORTType.Reset();
+		TensorNameValues.Reset();
+		TensorNames.Reset();
+		SymbolicDimensionNames.SetNum(NumberTensors);
+
+		for (uint32 TensorIndex = 0; TensorIndex < NumberTensors; ++TensorIndex)
+		{
+			// Get Tensor name
+			Ort::AllocatedStringPtr CurTensorName = bAreTensorInputs ? Session->GetInputNameAllocated(TensorIndex, *Allocator) : Session->GetOutputNameAllocated(TensorIndex, *Allocator);
+			TensorNameValues.Emplace(MoveTemp(CurTensorName));
+			TensorNames.Emplace(TensorNameValues.Last().get());
+
+			// Get node type
+			const Ort::TypeInfo CurrentTypeInfo = bAreTensorInputs ? Session->GetInputTypeInfo(TensorIndex) : Session->GetOutputTypeInfo(TensorIndex);
+			const Ort::ConstTensorTypeAndShapeInfo CurrentTensorInfo = CurrentTypeInfo.GetTensorTypeAndShapeInfo();
+			const ONNXTensorElementDataType ONNXTensorElementDataTypeEnum = CurrentTensorInfo.GetElementType();
+			const TypeInfoORT TypeInfo = TranslateTensorTypeORTToNNE(ONNXTensorElementDataTypeEnum);
+
+			// Get dynamic shape dimension names
+			TUniquePtr<const char*[]> SymbolidDimensionNames = MakeUnique<const char*[]>(CurrentTensorInfo.GetShape().size());
+			CurrentTensorInfo.GetSymbolicDimensions(SymbolidDimensionNames.Get(), CurrentTensorInfo.GetShape().size());
+
+			TArray<FString>& CurrentSymbolicDimensionNames = SymbolicDimensionNames[TensorIndex];
+			CurrentSymbolicDimensionNames.SetNum(CurrentTensorInfo.GetShape().size());
+
+			for (int32 i = 0; i < CurrentTensorInfo.GetShape().size(); i++)
+			{
+				CurrentSymbolicDimensionNames[i] = FString((SymbolidDimensionNames.Get())[i]);
+			}
+
+			TensorsORTType.Emplace(ONNXTensorElementDataTypeEnum);
+
+			// Get shape
+			TArray<int32, TInlineAllocator<NNE::FTensorShape::MaxRank>> ShapeData;
+			ShapeData.Reserve(CurrentTensorInfo.GetShape().size());
+			for (int64 CurrentTensorSize : CurrentTensorInfo.GetShape())
+			{
+				ShapeData.Add((int32)CurrentTensorSize);
+			}
+
+			const NNE::FSymbolicTensorShape Shape = NNE::FSymbolicTensorShape::Make(ShapeData);
+			const NNE::FTensorDesc SymbolicTensorDesc = NNE::FTensorDesc::Make(FString(TensorNames.Last()), Shape, TypeInfo.DataType);
+
+			check(SymbolicTensorDesc.GetElementByteSize() == TypeInfo.ElementSize);
+			SymbolicTensorDescs.Emplace(SymbolicTensorDesc);
+		}
+
+		return true;
+	}
+
+	void SetOutputTensorShapes(TConstArrayView<NNE::FTensorShape> Shapes)
+	{
+		FScopeLock Lock(&CriticalSection);
+
+		OutputTensorShapes = Shapes;
+	}
+
+	TConstArrayView<NNE::FTensorShape> GetOutputTensorShapes() const
+	{
+		FScopeLock Lock(&CriticalSection);
+
+		return OutputTensorShapes;
+	}
+
+	TSharedRef<UE::NNE::FSharedModelData> ModelData;
+	FRuntimeConf RuntimeConf;
+	FString TempDirForModelWithExternalData;
+
+	/** ORT-related variables */
+	TSharedRef<FEnvironment> Environment;
+	TUniquePtr<Ort::Session> Session;
+	TUniquePtr<Ort::SessionOptions> SessionOptions;
+	TUniquePtr<Ort::AllocatorWithDefaultOptions> Allocator;
+
+	TArray<NNE::FTensorDesc>	InputSymbolicTensors;	// Constant
+	TArray<NNE::FTensorDesc>	OutputSymbolicTensors;	// Should be constant, written/read by RHI Thread
+
+	/** IO ORT-related variables */
+	TArray<ONNXTensorElementDataType> InputTensorsORTType;
+	TArray<ONNXTensorElementDataType> OutputTensorsORTType;
+
+	TArray<Ort::AllocatedStringPtr> InputTensorNameValues;
+	TArray<Ort::AllocatedStringPtr> OutputTensorNameValues;
+	TArray<char*> InputTensorNames;
+	TArray<char*> OutputTensorNames;
+	TArray<TArray<FString>> InputSymbolicDimensionNames;
+	TArray<TArray<FString>> OutputSymbolicDimensionNames;
+
+	TArray<FTensor> OutputTensors;
+
+private:
+	mutable FCriticalSection CriticalSection;
+	TArray<NNE::FTensorShape>	OutputTensorShapes;		// Written on RHI Thread, read on Game/Render Thread -> Sync
+};
+
+FModelInstanceORTDmlRDG::FModelInstanceORTDmlRDG(TSharedRef<UE::NNE::FSharedModelData> InModelData, const FRuntimeConf& InRuntimeConf, TSharedRef<FEnvironment> InEnvironment)
+{
+	Proxy = MakeShared<FProxy>(InModelData, InRuntimeConf, InEnvironment);
+}
+
+FModelInstanceORTDmlRDG::~FModelInstanceORTDmlRDG()
+{
+	
 }
 
 bool FModelInstanceORTDmlRDG::Init()
@@ -597,28 +729,36 @@ bool FModelInstanceORTDmlRDG::Init()
 	try
 #endif // WITH_EDITOR
 	{
-		Allocator = MakeUnique<Ort::AllocatorWithDefaultOptions>();
+		Proxy->Allocator = MakeUnique<Ort::AllocatorWithDefaultOptions>();
 
-		SessionOptions = CreateSessionOptionsForDirectML(Environment);
-		if (!SessionOptions.IsValid())
+		Proxy->SessionOptions = CreateSessionOptionsForDirectML(Proxy->Environment);
+		if (!Proxy->SessionOptions.IsValid())
 		{
 			UE_LOG(LogNNERuntimeORT, Error, TEXT("Failed to configure session options for DirectML Execution Provider."));
 			return false;
 		}
 		
-		SessionOptions->SetGraphOptimizationLevel(GetGraphOptimizationLevelForDML(true));
+		Proxy->SessionOptions->SetGraphOptimizationLevel(GetGraphOptimizationLevelForDML(true));
 
-		if (!Detail::CreateSession(ModelData->GetView(), *SessionOptions, *Environment, Session, TempDirForModelWithExternalData))
+		if (!Detail::CreateSession(Proxy->ModelData->GetView(), *Proxy->SessionOptions, *Proxy->Environment, Proxy->Session, Proxy->TempDirForModelWithExternalData))
 		{
 			UE_LOG(LogNNERuntimeORT, Error, TEXT("Session creation failed."));
 			return false;
 		}
 
-		if (!ConfigureTensors(*Session))
+		if (!Proxy->ConfigureTensors(true))
 		{
 			UE_LOG(LogNNERuntimeORT, Error, TEXT("Failed to configure Inputs tensors."));
 			return false;
 		}
+		if (!Proxy->ConfigureTensors(false))
+		{
+			UE_LOG(LogNNERuntimeORT, Error, TEXT("Failed to configure Outputs tensors."));
+			return false;
+		}
+
+		InitialInputSymbolicTensors = Proxy->InputSymbolicTensors;
+		InitialOutputSymbolicTensors = Proxy->OutputSymbolicTensors;
 	}
 #if WITH_EDITOR
 	catch (const Ort::Exception& Exception)
@@ -636,213 +776,208 @@ bool FModelInstanceORTDmlRDG::Init()
 	return true;
 }
 
-bool FModelInstanceORTDmlRDG::ConfigureTensors(const Ort::Session& ActiveSession)
+TConstArrayView<NNE::FTensorDesc> FModelInstanceORTDmlRDG::GetInputTensorDescs() const
 {
-	if (!ConfigureTensors(ActiveSession, true))
-	{
-		UE_LOG(LogNNERuntimeORT, Error, TEXT("Failed to configure Inputs tensors."));
-		return false;
-	}
-	if (!ConfigureTensors(ActiveSession, false))
-	{
-		UE_LOG(LogNNERuntimeORT, Error, TEXT("Failed to configure Outputs tensors."));
-		return false;
-	}
-
-	return true;
+	return InitialInputSymbolicTensors;
 }
 
-bool FModelInstanceORTDmlRDG::ConfigureTensors(const Ort::Session& ActiveSession, bool bAreTensorInputs)
+TConstArrayView<NNE::FTensorShape> FModelInstanceORTDmlRDG::GetInputTensorShapes() const
 {
-	SCOPED_NAMED_EVENT_TEXT("FModelInstanceORTDmlRDG::ConfigureTensors", FColor::Magenta);
-
-	const uint32 NumberTensors							= bAreTensorInputs ? ActiveSession.GetInputCount() : ActiveSession.GetOutputCount();
-	TArray<NNE::FTensorDesc>& SymbolicTensorDescs		= bAreTensorInputs ? InputSymbolicTensors   : OutputSymbolicTensors;
-	TArray<ONNXTensorElementDataType>& TensorsORTType	= bAreTensorInputs ? InputTensorsORTType	: OutputTensorsORTType;
-	TArray<char*>& TensorNames							= bAreTensorInputs ? InputTensorNames		: OutputTensorNames;
-	TArray<Ort::AllocatedStringPtr>& TensorNameValues	= bAreTensorInputs ? InputTensorNameValues	: OutputTensorNameValues;
-	TArray<TArray<FString>>& SymbolicDimensionNames 	= bAreTensorInputs ? InputSymbolicDimensionNames : OutputSymbolicDimensionNames;
+	return InputTensorShapes;
+}
 
-	SymbolicTensorDescs.Reset();
-	TensorsORTType.Reset();
-	TensorNameValues.Reset();
-	TensorNames.Reset();
-	SymbolicDimensionNames.SetNum(NumberTensors);
+TConstArrayView<NNE::FTensorDesc> FModelInstanceORTDmlRDG::GetOutputTensorDescs() const
+{
+	return InitialOutputSymbolicTensors;
+}
 
-	for (uint32 TensorIndex = 0; TensorIndex < NumberTensors; ++TensorIndex)
-	{
-		// Get Tensor name
-		Ort::AllocatedStringPtr CurTensorName = bAreTensorInputs ? ActiveSession.GetInputNameAllocated(TensorIndex, *Allocator) : ActiveSession.GetOutputNameAllocated(TensorIndex, *Allocator);
-		TensorNameValues.Emplace(MoveTemp(CurTensorName));
-		TensorNames.Emplace(TensorNameValues.Last().get());
+TConstArrayView<NNE::FTensorShape> FModelInstanceORTDmlRDG::GetOutputTensorShapes() const
+{
+	return Proxy->GetOutputTensorShapes();
+}
 
-		// Get node type
-		const Ort::TypeInfo CurrentTypeInfo = bAreTensorInputs ? ActiveSession.GetInputTypeInfo(TensorIndex) : ActiveSession.GetOutputTypeInfo(TensorIndex);
-		const Ort::ConstTensorTypeAndShapeInfo CurrentTensorInfo = CurrentTypeInfo.GetTensorTypeAndShapeInfo();
-		const ONNXTensorElementDataType ONNXTensorElementDataTypeEnum = CurrentTensorInfo.GetElementType();
-		const TypeInfoORT TypeInfo = TranslateTensorTypeORTToNNE(ONNXTensorElementDataTypeEnum);
 
-		// Get dynamic shape dimension names
-		TUniquePtr<const char*[]> SymbolidDimensionNames = MakeUnique<const char*[]>(CurrentTensorInfo.GetShape().size());
-		CurrentTensorInfo.GetSymbolicDimensions(SymbolidDimensionNames.Get(), CurrentTensorInfo.GetShape().size());
+FModelInstanceORTDmlRDG::ESetInputTensorShapesStatus FModelInstanceORTDmlRDG::SetInputTensorShapes(TConstArrayView<NNE::FTensorShape> InInputShapes)
+{
+	SCOPED_NAMED_EVENT_TEXT("FModelInstanceORTDmlRDG::SetInputTensorShapes", FColor::Magenta);
 
-		TArray<FString>& CurrentSymbolicDimensionNames = SymbolicDimensionNames[TensorIndex];
-		CurrentSymbolicDimensionNames.SetNum(CurrentTensorInfo.GetShape().size());
+	if (IsInRenderingThread())
+	{
+		return SetInputTensorShapes_RenderThread(InInputShapes);
+	}
+	else
+	{
+		check(IsInGameThread());
 
-		for (int32 i = 0; i < CurrentTensorInfo.GetShape().size(); i++)
+		FEvent*	Signal = FGenericPlatformProcess::GetSynchEventFromPool();
+		
+		ENQUEUE_RENDER_COMMAND(FModelInstanceORTDmlRDG_SetInputTensorShapes)([this, InInputShapes, Signal](FRHICommandListImmediate& RHICmdList)
 		{
-			CurrentSymbolicDimensionNames[i] = FString((SymbolidDimensionNames.Get())[i]);
-		}
+			SetInputTensorShapes_RenderThread(InInputShapes);
 
-		TensorsORTType.Emplace(ONNXTensorElementDataTypeEnum);
-
-		// Get shape
-		TArray<int32, TInlineAllocator<NNE::FTensorShape::MaxRank>> ShapeData;
-		ShapeData.Reserve(CurrentTensorInfo.GetShape().size());
-		for (int64 CurrentTensorSize : CurrentTensorInfo.GetShape())
-		{
-			ShapeData.Add((int32)CurrentTensorSize);
-		}
+			// note: Block here, if SetInputTensorShapes does actual GPU work!
+			Signal->Trigger();
+		});
 
-		const NNE::FSymbolicTensorShape Shape = NNE::FSymbolicTensorShape::Make(ShapeData);
-		const NNE::FTensorDesc SymbolicTensorDesc = NNE::FTensorDesc::Make(FString(TensorNames.Last()), Shape, TypeInfo.DataType);
+		Signal->Wait();
 
-		check(SymbolicTensorDesc.GetElementByteSize() == TypeInfo.ElementSize);
-		SymbolicTensorDescs.Emplace(SymbolicTensorDesc);
+		FGenericPlatformProcess::ReturnSynchEventToPool(Signal);
 	}
 
-	return true;
+	return ESetInputTensorShapesStatus::Ok;
 }
 
-FModelInstanceORTDmlRDG::ESetInputTensorShapesStatus FModelInstanceORTDmlRDG::SetInputTensorShapes(TConstArrayView<NNE::FTensorShape> InInputShapes)
+// note: We delay shape inference to same RHI thread that runs the ORT session. This way we keep order of execution, but multiple calls to
+// SetInputTensorShapes() followed by EnqueueRDG() do not interfere each other.
+// SetInputTensorShapes() and EnqueueRDG() need to be called from the same thread, otherwise behaviour is undefined.
+// SetInputTensorShapes() sets following members immediately, since they are accessible from the outside:
+// - InputTensorShapes
+// It also sets InputTensors, since its used by EnqueueRDG() on the same thread, but then needs to be preserved.
+// OutputTensorShapes access requires synchronization, since its set by RHI thread and might be read by game/render thread
+FModelInstanceORTDmlRDG::ESetInputTensorShapesStatus FModelInstanceORTDmlRDG::SetInputTensorShapes_RenderThread(TConstArrayView<NNE::FTensorShape> InInputShapes)
 {
-	SCOPED_NAMED_EVENT_TEXT("FModelInstanceORTDmlRDG::SetInputTensorShapes", FColor::Magenta);
+	SCOPED_NAMED_EVENT_TEXT("FModelInstanceORTDmlRDG::SetInputTensorShapes_RenderThread", FColor::Magenta);
 
-	InputTensors.Reset();
-	OutputTensors.Reset();
-	OutputTensorShapes.Reset();
+	check(IsInRenderingThread());
+
+	Proxy->SetOutputTensorShapes({});
 
 	// Verify input shape are valid for the model and set InputTensorShapes
-	if (ESetInputTensorShapesStatus Status = FModelInstanceORTBase<NNE::IModelInstanceRDG>::SetInputTensorShapes(InInputShapes); Status != ESetInputTensorShapesStatus::Ok)
+	if (!Detail::SetInputTensorShapes(InInputShapes, InitialInputSymbolicTensors, InputTensorShapes))
 	{
-		return Status;
+		return ESetInputTensorShapesStatus::Fail;
 	}
 
-	// Check whether all input tensor shapes are concrete
-	bool bHasSymbolicInputShapes = false;
-	for (int32 i = 0; i < InputSymbolicTensors.Num(); i++)
+	// Set concrete input shapes
+	InputTensors.Reset();
+	for (int32 i = 0; i < InitialInputSymbolicTensors.Num(); i++)
 	{
-		if (!InputSymbolicTensors[i].GetShape().IsConcrete())
-		{
-			bHasSymbolicInputShapes = true;
-
-			break;
-		}
+		FTensor Tensor = FTensor::Make(InputTensorShapes[i], InitialInputSymbolicTensors[i].GetDataType());
+		InputTensors.Emplace(Tensor);
 	}
 
-	if (!bHasSymbolicInputShapes)
+	FRHICommandListImmediate &RHICmdList = GetImmediateCommandList_ForRenderCommand();
+	RHICmdList.EnqueueLambda([Proxy = Proxy, InputTensorShapes = InputTensorShapes](FRHICommandListImmediate& RHICmdList)
 	{
-		for (int32 i = 0; i < InputSymbolicTensors.Num(); i++)
+		GetID3D12PlatformDynamicRHI()->RHIRunOnQueue(ED3D12RHIRunOnQueueType::Graphics, [Proxy = Proxy, InputTensorShapes = InputTensorShapes] (ID3D12CommandQueue* D3D12CommandQueue)
 		{
-			FTensor Tensor = FTensor::Make(InInputShapes[i], InputSymbolicTensors[i].GetDataType());
-			InputTensors.Emplace(Tensor);
-		}
+			SCOPED_NAMED_EVENT_TEXT("FModelInstanceORTDmlRDG::SetInputTensorShapes RHIRunOnQueue", FColor::Magenta);
 
-		// All output shapes need to be concrete now
-		for (int32 i = 0; i < OutputSymbolicTensors.Num(); i++)
-		{
-			const NNE::FTensorDesc SymbolicTensorDesc = OutputSymbolicTensors[i];
+			Proxy->OutputTensors.Reset();
+			TArray<NNE::FTensorShape>	ResultOutputTensorShapes;
 
-			if (SymbolicTensorDesc.GetShape().IsConcrete())
+			// Check whether all input tensor shapes are concrete
+			bool bHasSymbolicInputShapes = false;
+			for (int32 i = 0; i < Proxy->InputSymbolicTensors.Num(); i++)
 			{
-				FTensor Tensor = FTensor::MakeFromSymbolicDesc(SymbolicTensorDesc);
-				OutputTensors.Emplace(Tensor);
-				OutputTensorShapes.Emplace(Tensor.GetShape());
-			}
-			else
-			{
-				UE_LOG(LogNNERuntimeORT, Warning, TEXT("One or more output tensors contain free dimensions, but input tensors are all concrete!"));
-				return ESetInputTensorShapesStatus::Fail;
-			}
-		}
+				if (!Proxy->InputSymbolicTensors[i].GetShape().IsConcrete())
+				{
+					bHasSymbolicInputShapes = true;
 
-		return ESetInputTensorShapesStatus::Ok;
-	}
+					break;
+				}
+			}
 
-	// Recreate session options because potentially we add new free dimension overrides
-	SessionOptions = CreateSessionOptionsForDirectML(Environment);
-	if (!SessionOptions.IsValid())
-	{
-		UE_LOG(LogNNERuntimeORT, Error, TEXT("Failed to recreate session options!"));
-		return ESetInputTensorShapesStatus::Fail;
-	}
+			if (!bHasSymbolicInputShapes)
+			{
+				// All output shapes need to be concrete now
+				for (int32 i = 0; i < Proxy->OutputSymbolicTensors.Num(); i++)
+				{
+					const NNE::FTensorDesc SymbolicTensorDesc = Proxy->OutputSymbolicTensors[i];
 
-	SessionOptions->SetGraphOptimizationLevel(GetGraphOptimizationLevelForDML(true));
+					if (SymbolicTensorDesc.GetShape().IsConcrete())
+					{
+						FTensor Tensor = FTensor::MakeFromSymbolicDesc(SymbolicTensorDesc);
+						Proxy->OutputTensors.Emplace(Tensor);
+						ResultOutputTensorShapes.Emplace(Tensor.GetShape());
+					}
+					else
+					{
+						UE_LOG(LogNNERuntimeORT, Warning, TEXT("One or more output tensors contain free dimensions, but input tensors are all concrete!"));
+						return;
+					}
+				}
 
-	// Setup concrete input tensors
-	for (int32 i = 0; i < InputSymbolicTensors.Num(); i++)
-	{
-		FTensor Tensor = FTensor::Make(InInputShapes[i], InputSymbolicTensors[i].GetDataType());
-		InputTensors.Emplace(Tensor);
+				Proxy->SetOutputTensorShapes(ResultOutputTensorShapes);
 
-		const NNE::FSymbolicTensorShape& SymbolicInputShape = InputSymbolicTensors[i].GetShape();
+				return;
+			}
 
-		// Override free dimensions of input tensors
-		if (!SymbolicInputShape.IsConcrete())
-		{
-			check(InputTensorShapes[i].IsCompatibleWith(SymbolicInputShape));
+			// Recreate session options because potentially we add new free dimension overrides
+			Proxy->SessionOptions = CreateSessionOptionsForDirectML(Proxy->Environment);
+			if (!Proxy->SessionOptions.IsValid())
+			{
+				UE_LOG(LogNNERuntimeORT, Error, TEXT("Failed to recreate session options!"));
+				return;
+			}
 
-			TConstArrayView<int32> InputSymbolicShapeData = SymbolicInputShape.GetData();
-			TConstArrayView<uint32> InputShapeData = InputTensorShapes[i].GetData();
+			Proxy->SessionOptions->SetGraphOptimizationLevel(GetGraphOptimizationLevelForDML(true));
 
-			for (int32 j = 0; j < InputShapeData.Num(); j++)
+			// Setup concrete input tensors
+			for (int32 i = 0; i < Proxy->InputSymbolicTensors.Num(); i++)
 			{
-				if (InputSymbolicShapeData[j] < 0)
+				const NNE::FSymbolicTensorShape& SymbolicInputShape = Proxy->InputSymbolicTensors[i].GetShape();
+
+				// Override free dimensions of input tensors
+				if (!SymbolicInputShape.IsConcrete())
 				{
-					Ort::GetApi().AddFreeDimensionOverrideByName(*SessionOptions, TCHAR_TO_ANSI(*InputSymbolicDimensionNames[i][j]), InputShapeData[j]);
+					check(InputTensorShapes[i].IsCompatibleWith(SymbolicInputShape));
+
+					TConstArrayView<int32> InputSymbolicShapeData = SymbolicInputShape.GetData();
+					TConstArrayView<uint32> InputShapeData = InputTensorShapes[i].GetData();
+
+					for (int32 j = 0; j < InputShapeData.Num(); j++)
+					{
+						if (InputSymbolicShapeData[j] < 0)
+						{
+							Ort::GetApi().AddFreeDimensionOverrideByName(*Proxy->SessionOptions, TCHAR_TO_ANSI(*Proxy->InputSymbolicDimensionNames[i][j]), InputShapeData[j]);
+						}
+					}
 				}
 			}
-		}
-	}
-
-	if (!Detail::CreateSession(ModelData->GetView(), *SessionOptions, *Environment, Session, TempDirForModelWithExternalData))
-	{
-		UE_LOG(LogNNERuntimeORT, Error, TEXT("Failed to recreate session!"));
-		return ESetInputTensorShapesStatus::Fail;
-	}
 
-	// Need to configure output tensors with new session (to apply free dimension overrides)
-	if (!ConfigureTensors(*Session, false))
-	{
-		UE_LOG(LogNNERuntimeORT, Error, TEXT("Failed to configure tensors!"));
-		return ESetInputTensorShapesStatus::Fail;
-	}
+			if (!Detail::CreateSession(Proxy->ModelData->GetView(), *Proxy->SessionOptions, *Proxy->Environment, Proxy->Session, Proxy->TempDirForModelWithExternalData))
+			{
+				UE_LOG(LogNNERuntimeORT, Error, TEXT("Failed to recreate session!"));
+				return;
+			}
 
-	// All output shapes need to be concrete now
-	for (int32 i = 0; i < OutputSymbolicTensors.Num(); i++)
-	{
-		const NNE::FTensorDesc SymbolicTensorDesc = OutputSymbolicTensors[i];
+			// Need to configure output tensors with new session (to apply free dimension overrides)
+			if (!Proxy->ConfigureTensors(false))
+			{
+				UE_LOG(LogNNERuntimeORT, Error, TEXT("Failed to configure tensors!"));
+				return;
+			}
 
-		if (SymbolicTensorDesc.GetShape().IsConcrete())
-		{
-			FTensor Tensor = FTensor::MakeFromSymbolicDesc(SymbolicTensorDesc);
-			OutputTensors.Emplace(Tensor);
-			OutputTensorShapes.Emplace(Tensor.GetShape());
-		}
-		else
-		{
-			for (int32 j = 0; j < SymbolicTensorDesc.GetShape().Rank(); j++)
+			// All output shapes need to be concrete now
+			for (int32 i = 0; i < Proxy->OutputSymbolicTensors.Num(); i++)
 			{
-				if (SymbolicTensorDesc.GetShape().GetData()[j] < 0)
+				const NNE::FTensorDesc SymbolicTensorDesc = Proxy->OutputSymbolicTensors[i];
+
+				if (SymbolicTensorDesc.GetShape().IsConcrete())
+				{
+					FTensor Tensor = FTensor::MakeFromSymbolicDesc(SymbolicTensorDesc);
+					Proxy->OutputTensors.Emplace(Tensor);
+					ResultOutputTensorShapes.Emplace(Tensor.GetShape());
+				}
+				else
 				{
-					UE_LOG(LogNNERuntimeORT, Warning, TEXT("Tensor '%hs' has free dimension '%s'."), OutputTensorNames[i], *OutputSymbolicDimensionNames[i][j]);
+					for (int32 j = 0; j < SymbolicTensorDesc.GetShape().Rank(); j++)
+					{
+						if (SymbolicTensorDesc.GetShape().GetData()[j] < 0)
+						{
+							UE_LOG(LogNNERuntimeORT, Warning, TEXT("Tensor '%hs' has free dimension '%s'."), Proxy->OutputTensorNames[i], *Proxy->OutputSymbolicDimensionNames[i][j]);
+						}
+					}
+
+					UE_LOG(LogNNERuntimeORT, Error, TEXT("One or more output tensors contain free dimensions!"));
+					return;
 				}
 			}
 
-			UE_LOG(LogNNERuntimeORT, Error, TEXT("One or more output tensors contain free dimensions!"));
-			return ESetInputTensorShapesStatus::Fail;
-		}
-	}
+			Proxy->SetOutputTensorShapes(ResultOutputTensorShapes);
+		}, false);
+	});
 
 	return ESetInputTensorShapesStatus::Ok;
 }
@@ -882,37 +1017,29 @@ Ort::Value CreateTensor(const OrtDmlApi& DmlApi, const Ort::MemoryInfo& MemoryIn
 	return Result;
 }
 
+// note: make sure to copy anything read by the RHI thread, but set by the calling thread. Currently only member InputTensors is affected.
+// Everything else is constant (set on Init()) or set in the same RHI thread in the lambda enqueued by SetInputTensorShapes().
+// As of today and without any additional settings/flags, the lambda passed to RHICmdList.EnqueueLambda([](){}) is executed on the RHI thread and
+// the lambda passed to GetID3D12PlatformDynamicRHI()->RHIRunOnQueue(..., [](){}, ...) is executed on the RHI submission thread.
 FModelInstanceORTDmlRDG::EEnqueueRDGStatus FModelInstanceORTDmlRDG::EnqueueRDG(FRDGBuilder& GraphBuilder, TConstArrayView<NNE::FTensorBindingRDG> Inputs, TConstArrayView<NNE::FTensorBindingRDG> Outputs)
 {
 	SCOPED_NAMED_EVENT_TEXT("FModelInstanceORTDmlRDG::EnqueueRDG", FColor::Magenta);
 
-	if (!Session.IsValid())
-	{
-		UE_LOG(LogNNERuntimeORT, Error, TEXT("Invalid Session, Init() should have been called."));
-		return EEnqueueRDGStatus::Fail;
-	}
-
 	if (InputTensorShapes.IsEmpty())
 	{
 		UE_LOG(LogNNERuntimeORT, Error, TEXT("Input shapes are not set, please call SetInputTensorShapes."));
 		return EEnqueueRDGStatus::Fail;
 	}
 
-	check(InputTensorShapes.Num() == InputTensors.Num());
-	check(InputTensorShapes.Num() == InputTensorNames.Num());
-	check(InputSymbolicTensors.Num() == InputTensors.Num());
-
-	if (Inputs.Num() != InputTensors.Num())
+	if (Inputs.Num() != InitialInputSymbolicTensors.Num())
 	{
-		UE_LOG(LogNNERuntimeORT, Error, TEXT("Input bindings need to match input tensor descriptor count (got %d, expected %d)."), Inputs.Num(), InputTensors.Num());
+		UE_LOG(LogNNERuntimeORT, Error, TEXT("Input bindings need to match input tensor descriptor count (got %d, expected %d)."), Inputs.Num(), InitialInputSymbolicTensors.Num());
 		return EEnqueueRDGStatus::Fail;
 	}
 
-	check(OutputSymbolicTensors.Num() == OutputTensorNames.Num());
-
-	if (!Outputs.IsEmpty() && Outputs.Num() != OutputTensorNames.Num())
+	if (!Outputs.IsEmpty() && Outputs.Num() != InitialOutputSymbolicTensors.Num())
 	{
-		UE_LOG(LogNNERuntimeORT, Error, TEXT("Output binding can be empty or needs to match output tensor descriptor count (got %d, expected %d)."), Outputs.Num(), OutputTensorNames.Num());
+		UE_LOG(LogNNERuntimeORT, Error, TEXT("Output binding can be empty or needs to match output tensor descriptor count (got %d, expected %d)."), Outputs.Num(), InitialOutputSymbolicTensors.Num());
 		return EEnqueueRDGStatus::Fail;
 	}
 
@@ -926,37 +1053,19 @@ FModelInstanceORTDmlRDG::EEnqueueRDGStatus FModelInstanceORTDmlRDG::EnqueueRDG(F
 			return EEnqueueRDGStatus::Fail;
 		}
 
-		const uint64 DmlImpliedSizeBytes = CalcRDGBufferSizeForDirectML(InputTensors[i].GetDataSize());
-		if (Binding.Buffer && Binding.Buffer->Desc.GetSize() != DmlImpliedSizeBytes)
-		{
-			UE_LOG(LogNNERuntimeORT, Error, TEXT("Binding input tensor %d size does not match tensor buffer size required by DirectML (got %d, expected %" UINT64_FMT ", data size was %" UINT64_FMT ")."), i, Binding.Buffer->Desc.GetSize(), DmlImpliedSizeBytes, InputTensors[i].GetDataSize());
-			return EEnqueueRDGStatus::Fail;
-		}
-
 		PassParameters->InputBuffers.Emplace(Binding.Buffer, ERHIAccess::CopySrc);
 	}
 
-	TArray<int32> ValidOutputs;
 	for (int32 i = 0; i < Outputs.Num(); i++)
 	{
-		const NNE::FTensorBindingRDG& Binding = Outputs[i];
-
-		const uint64 DmlImpliedSizeBytes = CalcRDGBufferSizeForDirectML(OutputTensors[i].GetDataSize());
-		if (Binding.Buffer && Binding.Buffer->Desc.GetSize() != DmlImpliedSizeBytes)
-		{
-			UE_LOG(LogNNERuntimeORT, Error, TEXT("Binding output tensor %d size does not match tensor buffer size required by DirectML (got %d, expected %" UINT64_FMT ", data size was %" UINT64_FMT ")."), i, Binding.Buffer->Desc.GetSize(), DmlImpliedSizeBytes, OutputTensors[i].GetDataSize());
-			return EEnqueueRDGStatus::Fail;
-		}
-
-		PassParameters->OutputBuffers.Emplace(Binding.Buffer, ERHIAccess::CopyDest);
-		ValidOutputs.Add(i);
+		PassParameters->OutputBuffers.Emplace(Outputs[i].Buffer, ERHIAccess::CopyDest);
 	}
 
 	RDG_EVENT_SCOPE_STAT(GraphBuilder, FNNERuntimeORTDmlRDG, "FModelInstanceORTDmlRDG::EnqueueRDG");
 	RDG_GPU_STAT_SCOPE(GraphBuilder, FNNERuntimeORTDmlRDG);
 
 	GraphBuilder.AddPass(RDG_EVENT_NAME("FModelInstanceORTDmlRDG::EnqueueRDG.AddPass"), PassParameters, ERDGPassFlags::Readback,
-	[this, ValidOutputsCopy = ValidOutputs, PassParameters](FRHICommandListImmediate& RHICmdList)
+	[Proxy = Proxy, InputTensors = InputTensors, PassParameters](FRHICommandListImmediate& RHICmdList)
 	{
 		SCOPED_NAMED_EVENT_TEXT("FModelInstanceORTDmlRDG::EnqueueRDG.AddPass", FColor::Magenta);
 
@@ -967,26 +1076,31 @@ FModelInstanceORTDmlRDG::EEnqueueRDGStatus FModelInstanceORTDmlRDG::EnqueueRDG(F
 			InputBuffers[i] = PassParameters->InputBuffers[i]->GetRHI();
 		}
 
-		check(ValidOutputsCopy.Num() == PassParameters->OutputBuffers.Num());
-
 		TArray<FRHIBuffer*> OutputBuffers;
-		OutputBuffers.SetNumZeroed(OutputSymbolicTensors.Num());
-		for (int32 i = 0; i < ValidOutputsCopy.Num(); i++)
+		OutputBuffers.SetNumZeroed(Proxy->OutputSymbolicTensors.Num());
+		for (int32 i = 0; i < PassParameters->OutputBuffers.Num(); i++)
 		{
-			OutputBuffers[ValidOutputsCopy[i]] = PassParameters->OutputBuffers[i]->GetRHI();
+			OutputBuffers[i] = PassParameters->OutputBuffers[i]->GetRHI();
 		}
 
 		// Submit previous work here to the GPU to avoid ORT Session Run() dispatching its work first
 		RHICmdList.ImmediateFlush(EImmediateFlushType::DispatchToRHIThread);
 
-		RHICmdList.EnqueueLambda([this, InputBuffersCopy = InputBuffers, OutputBuffersCopy = OutputBuffers](FRHICommandListImmediate& RHICmdList)
+		RHICmdList.EnqueueLambda([Proxy = Proxy, InputTensors = InputTensors, InputBuffers = InputBuffers, OutputBuffers = OutputBuffers] (FRHICommandListImmediate& RHICmdList)
 		{
-			GetID3D12PlatformDynamicRHI()->RHIRunOnQueue(ED3D12RHIRunOnQueueType::Graphics, [this, InputBuffersCopyCopy = InputBuffersCopy, OutputBuffersCopyCopy = OutputBuffersCopy](ID3D12CommandQueue* D3D12CommandQueue)
+			GetID3D12PlatformDynamicRHI()->RHIRunOnQueue(ED3D12RHIRunOnQueueType::Graphics, [Proxy = Proxy, InputTensors = InputTensors, InputBuffers = InputBuffers, OutputBuffers = OutputBuffers] (ID3D12CommandQueue* D3D12CommandQueue)
 			{
+				SCOPED_NAMED_EVENT_TEXT("FModelInstanceORTDmlRDG::EnqueueRDG.AddPass RHIRunOnQueue", FColor::Magenta);
 #if WITH_EDITOR
 				try
 #endif // WITH_EDITOR
 				{
+					if (!Proxy->Session.IsValid())
+					{
+						UE_LOG(LogNNERuntimeORT, Error, TEXT("Invalid Session, may be Init() should have been called."));
+						return;
+					}
+
 					const OrtDmlApi* DmlApi;
 					Ort::ThrowOnError(Ort::GetApi().GetExecutionProviderApi("DML", ORT_API_VERSION, reinterpret_cast<const void**>(&DmlApi)));
 					
@@ -996,15 +1110,29 @@ FModelInstanceORTDmlRDG::EEnqueueRDGStatus FModelInstanceORTDmlRDG::EnqueueRDG(F
 					TArray<Ort::Value> OrtInputTensors;
 					TArray<Ort::Value> OrtOutputTensors;
 
-					for (int32 i = 0; i < InputBuffersCopyCopy.Num(); i++)
+					for (int32 i = 0; i < InputBuffers.Num(); i++)
 					{
-						OrtInputTensors.Add(CreateTensor(*DmlApi, MemoryInfo, InputBuffersCopyCopy[i], InputTensors[i], InputTensorsORTType[i], DmlAllocatorResources));
+						const uint64 DmlImpliedSizeBytes = CalcRDGBufferSizeForDirectML(InputTensors[i].GetDataSize());
+						if (InputBuffers[i] && InputBuffers[i]->GetDesc().Size < DmlImpliedSizeBytes)
+						{
+							UE_LOG(LogNNERuntimeORT, Error, TEXT("Binding input tensor %d size does not match tensor buffer size required by DirectML (got %d, expected %" UINT64_FMT ", data size was %" UINT64_FMT ")."), i, InputBuffers[i]->GetDesc().Size, DmlImpliedSizeBytes, InputTensors[i].GetDataSize());
+							return;
+						}
+						
+						OrtInputTensors.Add(CreateTensor(*DmlApi, MemoryInfo, InputBuffers[i], InputTensors[i], Proxy->InputTensorsORTType[i], DmlAllocatorResources));
 					}
-					for (int32 i = 0; i < OutputBuffersCopyCopy.Num(); i++)
+					for (int32 i = 0; i < OutputBuffers.Num(); i++)
 					{
-						if (OutputBuffersCopyCopy[i])
+						const uint64 DmlImpliedSizeBytes = CalcRDGBufferSizeForDirectML(Proxy->OutputTensors[i].GetDataSize());
+						if (OutputBuffers[i] && OutputBuffers[i]->GetDesc().Size < DmlImpliedSizeBytes)
+						{
+							UE_LOG(LogNNERuntimeORT, Error, TEXT("Binding output tensor %d size does not match tensor buffer size required by DirectML (got %d, expected %" UINT64_FMT ", data size was %" UINT64_FMT ")."), i, OutputBuffers[i]->GetDesc().Size, DmlImpliedSizeBytes, Proxy->OutputTensors[i].GetDataSize());
+							return;
+						}
+
+						if (OutputBuffers[i])
 						{
-							OrtOutputTensors.Add(CreateTensor(*DmlApi, MemoryInfo, OutputBuffersCopyCopy[i], OutputTensors[i], OutputTensorsORTType[i], DmlAllocatorResources));
+							OrtOutputTensors.Add(CreateTensor(*DmlApi, MemoryInfo, OutputBuffers[i], Proxy->OutputTensors[i], Proxy->OutputTensorsORTType[i], DmlAllocatorResources));
 						}
 						else
 						{
@@ -1012,9 +1140,9 @@ FModelInstanceORTDmlRDG::EEnqueueRDGStatus FModelInstanceORTDmlRDG::EnqueueRDG(F
 						}
 					}
 
-					Session->Run(Ort::RunOptions{ nullptr },
-						InputTensorNames.GetData(), &OrtInputTensors[0], InputTensorNames.Num(),
-						OutputTensorNames.GetData(), &OrtOutputTensors[0], OutputTensorNames.Num());
+					Proxy->Session->Run(Ort::RunOptions{ nullptr },
+						Proxy->InputTensorNames.GetData(), &OrtInputTensors[0], Proxy->InputTensorNames.Num(),
+						Proxy->OutputTensorNames.GetData(), &OrtOutputTensors[0], Proxy->OutputTensorNames.Num());
 				}
 #if WITH_EDITOR
 				catch (const Ort::Exception& Exception)
diff --git a/Engine/Plugins/NNE/NNERuntimeORT/Source/NNERuntimeORT/Private/NNERuntimeORTModel.h b/Engine/Plugins/NNE/NNERuntimeORT/Source/NNERuntimeORT/Private/NNERuntimeORTModel.h
--- a/Engine/Plugins/NNE/NNERuntimeORT/Source/NNERuntimeORT/Private/NNERuntimeORTModel.h
+++ b/Engine/Plugins/NNE/NNERuntimeORT/Source/NNERuntimeORT/Private/NNERuntimeORTModel.h
@@ -146,7 +146,7 @@ private:
 	TSharedRef<UE::NNE::FSharedModelData> ModelData;
 };
 
-class FModelInstanceORTDmlRDG : public FModelInstanceORTBase<NNE::IModelInstanceRDG>
+class FModelInstanceORTDmlRDG : public NNE::IModelInstanceRDG
 {
 
 public:
@@ -157,39 +157,27 @@ public:
 
 	virtual ~FModelInstanceORTDmlRDG();
 
-	virtual ESetInputTensorShapesStatus SetInputTensorShapes(TConstArrayView<NNE::FTensorShape> InInputShapes) override;
+	TConstArrayView<NNE::FTensorDesc> GetInputTensorDescs() const override;
+	TConstArrayView<NNE::FTensorDesc> GetOutputTensorDescs() const override;
+	TConstArrayView<NNE::FTensorShape> GetInputTensorShapes() const override;
+	TConstArrayView<NNE::FTensorShape> GetOutputTensorShapes() const override;
+
+	ESetInputTensorShapesStatus SetInputTensorShapes(TConstArrayView<NNE::FTensorShape> InInputShapes) override;
+	ESetInputTensorShapesStatus SetInputTensorShapes_RenderThread(TConstArrayView<NNE::FTensorShape> InInputShapes);
 
 	bool Init();
 
 	EEnqueueRDGStatus EnqueueRDG(FRDGBuilder& GraphBuilder, TConstArrayView<NNE::FTensorBindingRDG> Inputs, TConstArrayView<NNE::FTensorBindingRDG> Outputs) override;
 
-protected:
-	bool ConfigureTensors(const Ort::Session& ActiveSession);
-	bool ConfigureTensors(const Ort::Session& ActiveSession, bool InIsInput);
-
-	TSharedRef<UE::NNE::FSharedModelData> ModelData;
-	FRuntimeConf RuntimeConf;
-	FString TempDirForModelWithExternalData;
-
-	/** ORT-related variables */
-	TSharedRef<FEnvironment> Environment;
-	TUniquePtr<Ort::Session> Session;
-	TUniquePtr<Ort::SessionOptions> SessionOptions;
-	TUniquePtr<Ort::AllocatorWithDefaultOptions> Allocator;
+private:
+	TArray<NNE::FTensorShape>	InputTensorShapes; // Owned by Game/Render thread, copied for RHI thread
+	TArray<FTensor> InputTensors; // Owned by Game/Render thread, copied for RHI thread
 
-	/** IO ORT-related variables */
-	TArray<ONNXTensorElementDataType> InputTensorsORTType;
-	TArray<ONNXTensorElementDataType> OutputTensorsORTType;
+	TArray<NNE::FTensorDesc> InitialInputSymbolicTensors; // Copy for Game/Render Thread
+	TArray<NNE::FTensorDesc> InitialOutputSymbolicTensors; // Copy for Game/Render Thread
 
-	TArray<Ort::AllocatedStringPtr> InputTensorNameValues;
-	TArray<Ort::AllocatedStringPtr> OutputTensorNameValues;
-	TArray<char*> InputTensorNames;
-	TArray<char*> OutputTensorNames;
-	TArray<TArray<FString>> InputSymbolicDimensionNames;
-	TArray<TArray<FString>> OutputSymbolicDimensionNames;
-
-	TArray<FTensor> InputTensors;
-	TArray<FTensor> OutputTensors;
+	class FProxy; // Contains data written/read by RHI Thread and data initialized by Game/Render Thread
+	TSharedPtr<FProxy> Proxy;
 };
 
 class FModelInstanceORTNpu : public FModelInstanceORTRunSync<NNE::IModelInstanceNPU, NNE::FTensorBindingCPU>
