[NNE] NNERuntimeORT fixed link error with Windows SDK 10.0.26100.

#rnx

diff --git a/Engine/Plugins/NNE/NNERuntimeORT/Source/NNERuntimeORT/NNERuntimeORT.Build.cs b/Engine/Plugins/NNE/NNERuntimeORT/Source/NNERuntimeORT/NNERuntimeORT.Build.cs
--- a/Engine/Plugins/NNE/NNERuntimeORT/Source/NNERuntimeORT/NNERuntimeORT.Build.cs
+++ b/Engine/Plugins/NNE/NNERuntimeORT/Source/NNERuntimeORT/NNERuntimeORT.Build.cs
@@ -34,10 +34,11 @@ public class NNERuntimeORT : ModuleRules
 				"DX12"
 			});
 
-			if (Version.TryParse(Target.WindowsPlatform.WindowsSdkVersion, out Version WindowsSdkVersion) && WindowsSdkVersion.Build >= 26100 /* Windows 11, version 24H2 */)
-			{
-				PublicDefinitions.Add("ORT_USE_NEW_DXCORE_FEATURES");
-			}
+			// In Engine\Source\ThirdParty\Windows\DX12\DX12.Build.cs, we link against an older version of dxguid.lib, so we can not use the new one.
+			// if (Version.TryParse(Target.WindowsPlatform.WindowsSdkVersion, out Version WindowsSdkVersion) && WindowsSdkVersion.Build >= 26100 /* Windows 11, version 24H2 */)
+			// {
+			// 	PublicDefinitions.Add("ORT_USE_NEW_DXCORE_FEATURES");
+			// }
 		}
 
 		if (Target.bBuildEditor)
