[NNE] The DirectML NPU interface lists the GPU as NPU.

#tests Editor, Win64, verified correct reporting of NPU/GPU

#rnx

diff --git a/Engine/Plugins/NNE/NNERuntimeORT/Source/NNERuntimeORT/Private/NNERuntimeORTUtils.cpp b/Engine/Plugins/NNE/NNERuntimeORT/Source/NNERuntimeORT/Private/NNERuntimeORTUtils.cpp
--- a/Engine/Plugins/NNE/NNERuntimeORT/Source/NNERuntimeORT/Private/NNERuntimeORTUtils.cpp
+++ b/Engine/Plugins/NNE/NNERuntimeORT/Source/NNERuntimeORT/Private/NNERuntimeORTUtils.cpp
@@ -30,7 +30,7 @@ namespace UE::NNERuntimeORT::Private
 {
 
 #if PLATFORM_WINDOWS
-Microsoft::WRL::ComPtr<ID3D12Device1> CreateD3D12Device(IUnknown* AdapterPtr)
+Microsoft::WRL::ComPtr<ID3D12Device1> CreateD3D12Device(IUnknown* AdapterPtr, D3D_FEATURE_LEVEL FeatureLevel) 
 {
 	using Microsoft::WRL::ComPtr;
 
@@ -49,7 +49,7 @@ Microsoft::WRL::ComPtr<ID3D12Device1> CreateD3D12Device(IUnknown* AdapterPtr)
 	}
 
 	ComPtr<ID3D12Device1> D3D12Device;
-	HRESULT Hr = D3D12CreateDeviceFun(AdapterPtr, D3D_FEATURE_LEVEL_1_0_CORE, IID_PPV_ARGS(&D3D12Device));
+	HRESULT Hr = D3D12CreateDeviceFun(AdapterPtr, FeatureLevel, IID_PPV_ARGS(&D3D12Device));
 	if (FAILED(Hr))
 	{
 		UE_LOG(LogNNERuntimeORT, Log, TEXT("Failed to create D3D12 device, D3D12CreateDevice error code :%x"), Hr);
@@ -122,7 +122,7 @@ bool IsD3D12Available()
 		return false;
 	}
 
-	ComPtr<ID3D12Device> Device = CreateD3D12Device(Adapter.Get());
+	ComPtr<ID3D12Device> Device = CreateD3D12Device(Adapter.Get(), D3D_FEATURE_LEVEL_1_0_CORE);
 	if (!Device)
 	{
 		return false;
@@ -135,10 +135,12 @@ bool IsD3D12Available()
 }
 
 #if PLATFORM_WINDOWS
-Microsoft::WRL::ComPtr<IUnknown> GetAdapterNpu(bool bVerbose)
+Microsoft::WRL::ComPtr<IUnknown> GetAdapterNpu(bool bVerbose, D3D_FEATURE_LEVEL& OutFeatureLevel)
 {
 	using Microsoft::WRL::ComPtr;
 
+	OutFeatureLevel = D3D_FEATURE_LEVEL_1_0_CORE;
+
 	const bool bForceComputeOnlyDevice = true;
 
 	auto GetAdapterName = [](IDXCoreAdapter* Adapter) -> FString
@@ -196,41 +198,65 @@ Microsoft::WRL::ComPtr<IUnknown> GetAdapterNpu(bool bVerbose)
 	// New path from Windows 11, version 24H2 on
 	{
 #ifdef ORT_USE_NEW_DXCORE_FEATURES
-		const GUID DxGUIDs[] = { DXCORE_ADAPTER_ATTRIBUTE_D3D12_GENERIC_ML };
+		const GUID GUID_LIST_GENERIC_ML[] = { DXCORE_ADAPTER_ATTRIBUTE_D3D12_GENERIC_ML };
 #else
-		const GUID DxGUIDs[] = { { 0xb71b0d41, 0x1088, 0x422f, { 0xa2, 0x7c, 0x2, 0x50, 0xb7, 0xd3, 0xa9, 0x88 } } }; // Value of DXCORE_ADAPTER_ATTRIBUTE_D3D12_GENERIC_ML
+		const GUID GUID_LIST_GENERIC_ML[] = { { 0xb71b0d41, 0x1088, 0x422f, { 0xa2, 0x7c, 0x2, 0x50, 0xb7, 0xd3, 0xa9, 0x88 } } }; // Value of DXCORE_ADAPTER_ATTRIBUTE_D3D12_GENERIC_ML
 #endif
+		const GUID GUID_LIST_CORE_COMPUTE[] = { DXCORE_ADAPTER_ATTRIBUTE_D3D12_CORE_COMPUTE };
 
+		auto GetNonGraphicsAdapter = [] (IDXCoreAdapterList* AdapterList, IDXCoreAdapter** OutAdapter)
+		{
+			for (uint32_t i = 0, AdapterCount = AdapterList->GetAdapterCount(); i < AdapterCount; i++)
+			{
+				ComPtr<IDXCoreAdapter> Adapter;
+				HRESULT Hr = AdapterList->GetAdapter(i, IID_PPV_ARGS(&Adapter));
+				if (FAILED(Hr))
+				{
+					UE_LOG(LogNNERuntimeORT, Warning, TEXT("%d: Failed to get adapter, IDXCoreAdapterList::GetAdapter error code: %x"), i, Hr);
+				}
+
+				if (!Adapter->IsAttributeSupported(DXCORE_ADAPTER_ATTRIBUTE_D3D12_GRAPHICS))
+				{
+					*OutAdapter = Adapter.Detach();
+					return;
+				}
+			}
+			*OutAdapter = nullptr;
+		};
+
+		ComPtr<IDXCoreAdapter> Adapter;
 		ComPtr<IDXCoreAdapterList> AdapterList;
-		Hr = Factory->CreateAdapterList(ARRAYSIZE(DxGUIDs), DxGUIDs, IID_PPV_ARGS(&AdapterList));
+		D3D_FEATURE_LEVEL FeatureLevel = D3D_FEATURE_LEVEL_1_0_GENERIC;
+
+		Hr = Factory->CreateAdapterList(ARRAYSIZE(GUID_LIST_GENERIC_ML), GUID_LIST_GENERIC_ML, IID_PPV_ARGS(&AdapterList));
 		if (FAILED(Hr))
 		{
 			UE_LOG(LogNNERuntimeORT, Log, TEXT("Failed to create DXCore Adapter List, IDXCoreAdapterFactory::CreateAdapterList error code :%x"), Hr);
 			return {};
 		}
 
-		if (bVerbose)
+		if (AdapterList->GetAdapterCount() > 0)
 		{
-			UE_LOG(LogNNERuntimeORT, Log, TEXT("Available graphics and compute adapters:"));
+			GetNonGraphicsAdapter(AdapterList.Get(), Adapter.GetAddressOf());
 		}
 
-		ComPtr<IDXCoreAdapter> Adapter;
-		for (uint32_t i = 0, AdapterCount = AdapterList->GetAdapterCount(); i < AdapterCount; i++)
+		if (!Adapter)
 		{
-			Hr = AdapterList->GetAdapter(static_cast<uint32_t>(i), IID_PPV_ARGS(&Adapter));
+			FeatureLevel = D3D_FEATURE_LEVEL_1_0_CORE;
+
+			Hr = Factory->CreateAdapterList(ARRAYSIZE(GUID_LIST_CORE_COMPUTE), GUID_LIST_CORE_COMPUTE, IID_PPV_ARGS(&AdapterList));
 			if (FAILED(Hr))
 			{
-				UE_LOG(LogNNERuntimeORT, Warning, TEXT("%d: Failed to get adapter, IDXCoreAdapterList::GetAdapter error code: %x"), i, Hr);
-			}
-			else
-			{
-				break;
+				UE_LOG(LogNNERuntimeORT, Log, TEXT("Failed to create DXCore Adapter List, IDXCoreAdapterFactory::CreateAdapterList error code :%x"), Hr);
+				return {};
 			}
+
+			GetNonGraphicsAdapter(AdapterList.Get(), Adapter.GetAddressOf());
 		}
 
 		if (!Adapter)
 		{
-			UE_LOG(LogNNERuntimeORT, Log, TEXT("No NPU adapter found with attribute DXCORE_ADAPTER_ATTRIBUTE_D3D12_GENERIC_ML (Windows 11 Version 24H2 or newer)!"));
+			UE_LOG(LogNNERuntimeORT, Log, TEXT("No NPU adapter found with attributes (!GRAPHICS && (GENERIC_ML || CORE_COMPUTE)) (Windows 11 Version 24H2 or newer)!"));
 		}
 		else
 		{
@@ -239,6 +265,8 @@ Microsoft::WRL::ComPtr<IUnknown> GetAdapterNpu(bool bVerbose)
 				UE_LOG(LogNNERuntimeORT, Log, TEXT("Selecting NPU adapter: %s"), *(GetAdapterName(Adapter.Get())));
 			}
 
+			OutFeatureLevel = FeatureLevel;
+
 			return Adapter;
 		}
 	}
@@ -312,13 +340,15 @@ bool IsD3D12DeviceNPUAvailable()
 #if PLATFORM_WINDOWS
 	using Microsoft::WRL::ComPtr;
 
-	ComPtr<IUnknown> AdapterNpu = GetAdapterNpu(true);
+	D3D_FEATURE_LEVEL FeatureLevel = D3D_FEATURE_LEVEL_1_0_CORE;
+
+	ComPtr<IUnknown> AdapterNpu = GetAdapterNpu(true, FeatureLevel);
 	if (!AdapterNpu)
 	{
 		return false;
 	}
 
-	ComPtr<ID3D12Device1> D3D12Device = CreateD3D12Device(AdapterNpu.Get());
+	ComPtr<ID3D12Device1> D3D12Device = CreateD3D12Device(AdapterNpu.Get(), FeatureLevel);
 	if (!D3D12Device)
 	{
 		return false;
@@ -580,14 +610,16 @@ TUniquePtr<Ort::SessionOptions> CreateSessionOptionsForDirectMLNpu(const TShared
 	SessionOptions->SetExecutionMode(ExecutionMode::ORT_SEQUENTIAL);
 	SessionOptions->DisableMemPattern();
 
-	ComPtr<IUnknown> AdapterNpu = GetAdapterNpu(false);
+	D3D_FEATURE_LEVEL FeatureLevel = D3D_FEATURE_LEVEL_1_0_CORE;
+
+	ComPtr<IUnknown> AdapterNpu = GetAdapterNpu(false, FeatureLevel);
 	if (!AdapterNpu)
 	{
 		UE_LOG(LogNNERuntimeORT, Error, TEXT("Failed to get NPU adapter."));
 		return {};
 	}
 
-	ComPtr<ID3D12Device1> D3D12DeviceNpu = CreateD3D12Device(AdapterNpu.Get());
+	ComPtr<ID3D12Device1> D3D12DeviceNpu = CreateD3D12Device(AdapterNpu.Get(), FeatureLevel);
 	if (!D3D12DeviceNpu)
 	{
 		UE_LOG(LogNNERuntimeORT, Error, TEXT("Failed to create NPU device."));
