[NNE] NNERuntimeORT fix for new dxcore adapter listing.

#test Editor, Win64, verified correct reporting of GPU/NPU

#rnx

diff --git a/Engine/Plugins/NNE/NNERuntimeORT/Source/NNERuntimeORT/Private/NNERuntimeORTUtils.cpp b/Engine/Plugins/NNE/NNERuntimeORT/Source/NNERuntimeORT/Private/NNERuntimeORTUtils.cpp
--- a/Engine/Plugins/NNE/NNERuntimeORT/Source/NNERuntimeORT/Private/NNERuntimeORTUtils.cpp
+++ b/Engine/Plugins/NNE/NNERuntimeORT/Source/NNERuntimeORT/Private/NNERuntimeORTUtils.cpp
@@ -213,6 +213,7 @@ Microsoft::WRL::ComPtr<IUnknown> GetAdapterNpu(bool bVerbose, D3D_FEATURE_LEVEL&
 				if (FAILED(Hr))
 				{
 					UE_LOG(LogNNERuntimeORT, Warning, TEXT("%d: Failed to get adapter, IDXCoreAdapterList::GetAdapter error code: %x"), i, Hr);
+					continue;
 				}
 
 				if (!Adapter->IsAttributeSupported(DXCORE_ADAPTER_ATTRIBUTE_D3D12_GRAPHICS))
