[NNE] NNERuntimeORTDml defer NNE GPU runtime interface availability check.

#tests NNEAutomation tests on Win64, Game and Editor, all green.

#rnx

#Synced-CL 49674046

diff --git a/Engine/Plugins/NNE/NNERuntimeORT/Source/NNERuntimeORT/Private/NNERuntimeORT.cpp b/Engine/Plugins/NNE/NNERuntimeORT/Source/NNERuntimeORT/Private/NNERuntimeORT.cpp
--- a/Engine/Plugins/NNE/NNERuntimeORT/Source/NNERuntimeORT/Private/NNERuntimeORT.cpp
+++ b/Engine/Plugins/NNE/NNERuntimeORT/Source/NNERuntimeORT/Private/NNERuntimeORT.cpp
@@ -74,19 +74,29 @@ namespace UE::NNERuntimeORT::Private::Details
 		}
 	}
 
-	bool IsAvailableGPU(bool bDirectMLAvailable, bool bD3D12Available)
+	bool IsAvailableGPU(bool bIsDirectMLAvailable)
 	{
-		return bDirectMLAvailable && bD3D12Available;
+		static bool bIsD3D12Available = []() -> bool
+		{
+			return UE::NNERuntimeORT::Private::IsD3D12Available();
+		}();
+		return bIsDirectMLAvailable && bIsD3D12Available;
 	}
-
-	bool IsAvailableRDG(bool bDirectMLAvailable, bool bRHID3D12Available)
+	bool IsAvailableRDG(bool bIsDirectMLAvailable)
 	{
-		return bDirectMLAvailable && bRHID3D12Available;
+		static bool bIsRHID3D12Available = []() -> bool
+		{
+			return UE::NNERuntimeORT::Private::IsRHID3D12Available();
+		}();
+		return bIsDirectMLAvailable && bIsRHID3D12Available;
 	}
-
-	bool IsAvailableNPU(bool bDirectMLAvailable, bool bD3D12DeviceNPUAvailable)
+	bool IsAvailableNPU(bool bIsDirectMLAvailable)
 	{
-		return bDirectMLAvailable && bD3D12DeviceNPUAvailable;
+		static bool bIsD3D12DeviceNPUAvailable = []() -> bool
+		{
+			return UE::NNERuntimeORT::Private::IsD3D12DeviceNPUAvailable();
+		}();
+		return bIsDirectMLAvailable && bIsD3D12DeviceNPUAvailable;
 	}
 } // namespace UE::NNERuntimeORT::Private::Details
 
@@ -214,13 +224,8 @@ public:
 
 virtual void Init(TSharedRef<UE::NNERuntimeORT::Private::FEnvironment> InEnvironment, bool bInDirectMLAvailable) override
 {
-	using namespace UE::NNERuntimeORT::Private;
-
 	Environment = InEnvironment;
-	
-	bIsAvailableGPU = Details::IsAvailableGPU(bInDirectMLAvailable, IsD3D12Available());
-	bIsAvailableRDG = Details::IsAvailableRDG(bInDirectMLAvailable, IsRHID3D12Available());
-	bIsAvailableNPU = Details::IsAvailableNPU(bInDirectMLAvailable, IsD3D12DeviceNPUAvailable());
+	bIsDirectMLAvailable = bInDirectMLAvailable;
 }
 
 virtual FString GetRuntimeName() const override
@@ -284,7 +289,7 @@ virtual FString GetModelDataIdentifier(const FString& FileType, TConstArrayView6
 
 virtual ECanCreateModelGPUStatus CanCreateModelGPU(const TObjectPtr<UNNEModelData> ModelData) const override
 {
-	if (!ensureMsgf(bIsAvailableGPU, TEXT("GPU interface should not be available!")))
+	if (!ensureMsgf(UE::NNERuntimeORT::Private::Details::IsAvailableGPU(bIsDirectMLAvailable), TEXT("GPU interface should not be available!")))
 	{
 		return ECanCreateModelCommonStatus::Fail;
 	}
@@ -313,7 +318,7 @@ virtual TSharedPtr<UE::NNE::IModelGPU> CreateModelGPU(const TObjectPtr<UNNEModel
 
 virtual ECanCreateModelRDGStatus CanCreateModelRDG(TObjectPtr<UNNEModelData> ModelData) const override
 {
-	if (!ensureMsgf(bIsAvailableRDG, TEXT("RDG interface should not be available!")))
+	if (!ensureMsgf(UE::NNERuntimeORT::Private::Details::IsAvailableRDG(bIsDirectMLAvailable), TEXT("RDG interface should not be available!")))
 	{
 		return ECanCreateModelCommonStatus::Fail;
 	}
@@ -342,12 +347,12 @@ virtual TSharedPtr<UE::NNE::IModelRDG> CreateModelRDG(TObjectPtr<UNNEModelData>
 
 virtual ECanCreateModelNPUStatus CanCreateModelNPU(const TObjectPtr<UNNEModelData> ModelData) const override
 {
-	if (!ensureMsgf(bIsAvailableNPU, TEXT("NPU interface should not be available!")))
+	if (!ensureMsgf(UE::NNERuntimeORT::Private::Details::IsAvailableNPU(bIsDirectMLAvailable), TEXT("NPU interface should not be available!")))
 	{
 		return ECanCreateModelCommonStatus::Fail;
 	}
 
-	return CanCreateModelCommon(ModelData) == ECanCreateModelCommonStatus::Ok ? ECanCreateModelRDGStatus::Ok : ECanCreateModelRDGStatus::Fail;
+	return CanCreateModelCommon(ModelData) == ECanCreateModelCommonStatus::Ok ? ECanCreateModelNPUStatus::Ok : ECanCreateModelNPUStatus::Fail;
 }
 
 virtual TSharedPtr<UE::NNE::IModelNPU> CreateModelNPU(const TObjectPtr<UNNEModelData> ModelData) override
@@ -403,9 +408,7 @@ ECanCreateModelCommonStatus CanCreateModelCommon(const TObjectPtr<UNNEModelData>
 #endif // PLATFORM_WINDOWS
 }
 
-	bool bIsAvailableGPU = false;
-	bool bIsAvailableRDG = false;
-	bool bIsAvailableNPU = false;
+	bool bIsDirectMLAvailable = false;
 };
 
 FGuid FNNERuntimeORTDmlImpl::GUID = FGuid((int32)'O', (int32)'D', (int32)'M', (int32)'L');
@@ -429,32 +432,32 @@ namespace UE::NNERuntimeORT::Private
 
 TWeakObjectPtr<UNNERuntimeORTDmlProxy> MakeRuntimeDml(bool bDirectMLAvailable)
 {
+	// Important note: the check for GPU interface availability might be expensive, e.g. on Vulkan, because we need to check for D3D12 support.
+	// That will cause the DLL to load and perform other expensive work etc. This is undesired in the module loading phase and specifically
+	// we don't know whether its actually ever used.
+	// 
+	// Therefore, at this moment we exceptionally assume D3D12 "is present" and will do the check later in CanCreateModelGPU()
 	const bool bRHID3D12Available = IsRHID3D12Available();
-	const bool bD3D12Available = IsD3D12Available();
 	const bool bD3D12DeviceNPUAvailable = IsD3D12DeviceNPUAvailable();
 
-	const bool bIsAvailableGPU = Details::IsAvailableGPU(bDirectMLAvailable, bD3D12Available);
-	const bool bIsAvailableRDG = Details::IsAvailableRDG(bDirectMLAvailable, bRHID3D12Available);
-	const bool bIsAvailableNPU = Details::IsAvailableNPU(bDirectMLAvailable, bD3D12DeviceNPUAvailable);
+	const bool bIsAvailableRDG = bDirectMLAvailable && bRHID3D12Available;
+	const bool bIsAvailableNPU = bDirectMLAvailable && bD3D12DeviceNPUAvailable;
 
 	UE_LOG(LogNNERuntimeORT, Log, TEXT("MakeRuntimeORTDml:"));
 	UE_LOG(LogNNERuntimeORT, Log, TEXT("  DirectML:  %s"), (bDirectMLAvailable ? TEXT("yes") : TEXT("no")));
 	UE_LOG(LogNNERuntimeORT, Log, TEXT("  RHI D3D12: %s"), (bRHID3D12Available ? TEXT("yes") : TEXT("no")));
-	UE_LOG(LogNNERuntimeORT, Log, TEXT("  D3D12:     %s"), (bD3D12Available ? TEXT("yes") : TEXT("no")));
+	UE_LOG(LogNNERuntimeORT, Log, TEXT("  D3D12:     n/a (checked in CanCreateModelGPU)"));
 	UE_LOG(LogNNERuntimeORT, Log, TEXT("  NPU:       %s"), (bD3D12DeviceNPUAvailable ? TEXT("yes") : TEXT("no")));
 
 	UE_LOG(LogNNERuntimeORT, Log, TEXT("Interface availability:"));
-	UE_LOG(LogNNERuntimeORT, Log, TEXT("  GPU: %s"), (bIsAvailableGPU ? TEXT("yes") : TEXT("no")));
+	UE_LOG(LogNNERuntimeORT, Log, TEXT("  GPU: n/a (checked in CanCreateModelGPU)"));
 	UE_LOG(LogNNERuntimeORT, Log, TEXT("  RDG: %s"), (bIsAvailableRDG ? TEXT("yes") : TEXT("no")));
 	UE_LOG(LogNNERuntimeORT, Log, TEXT("  NPU: %s"), (bIsAvailableNPU ? TEXT("yes") : TEXT("no")));
 
-	if (bIsAvailableGPU && bIsAvailableRDG && bIsAvailableNPU) return NewObject<UNNERuntimeORTDml_GPU_RDG_NPU>();
-	else if (bIsAvailableGPU && bIsAvailableRDG) return NewObject<UNNERuntimeORTDml_GPU_RDG>();
-	else if (bIsAvailableGPU && bIsAvailableNPU) return NewObject<UNNERuntimeORTDml_GPU_NPU>();
-	else if (bIsAvailableRDG && bIsAvailableNPU) return NewObject<UNNERuntimeORTDml_RDG_NPU>();
-	else if (bIsAvailableGPU) return NewObject<UNNERuntimeORTDml_GPU>();
-	else if (bIsAvailableRDG) return NewObject<UNNERuntimeORTDml_RDG>();
-	else if (bIsAvailableNPU) return NewObject<UNNERuntimeORTDml_NPU>();
+	if (bIsAvailableRDG && bIsAvailableNPU) return NewObject<UNNERuntimeORTDml_GPU_RDG_NPU>();
+	else if (bIsAvailableRDG) return NewObject<UNNERuntimeORTDml_GPU_RDG>();
+	else if (bIsAvailableNPU) return NewObject<UNNERuntimeORTDml_GPU_NPU>();
+	else if (bDirectMLAvailable) return NewObject<UNNERuntimeORTDml_GPU>();
 
 #if WITH_EDITOR
 	UE_LOG(LogNNERuntimeORT, Log, TEXT("NNERuntimeORTDml can only cook!"));
diff --git a/Engine/Plugins/NNE/NNERuntimeORT/Source/NNERuntimeORT/Private/NNERuntimeORT.h b/Engine/Plugins/NNE/NNERuntimeORT/Source/NNERuntimeORT/Private/NNERuntimeORT.h
--- a/Engine/Plugins/NNE/NNERuntimeORT/Source/NNERuntimeORT/Private/NNERuntimeORT.h
+++ b/Engine/Plugins/NNE/NNERuntimeORT/Source/NNERuntimeORT/Private/NNERuntimeORT.h
@@ -115,19 +115,6 @@ public:
 	virtual TSharedPtr<UE::NNE::IModelNPU> CreateModelNPU(const TObjectPtr<UNNEModelData> ModelData) override { return Impl->CreateModelNPU(ModelData); }
 };
 
-UCLASS()
-class UNNERuntimeORTDml_RDG_NPU : public UNNERuntimeORTDmlProxy, public INNERuntimeRDG, public INNERuntimeNPU
-{
-	GENERATED_BODY()
-
-public:
-	virtual ECanCreateModelRDGStatus CanCreateModelRDG(TObjectPtr<UNNEModelData> ModelData) const override { return Impl->CanCreateModelRDG(ModelData); }
-	virtual TSharedPtr<UE::NNE::IModelRDG> CreateModelRDG(TObjectPtr<UNNEModelData> ModelData) override { return Impl->CreateModelRDG(ModelData); }
-
-	virtual ECanCreateModelNPUStatus CanCreateModelNPU(const TObjectPtr<UNNEModelData> ModelData) const override { return Impl->CanCreateModelNPU(ModelData); }
-	virtual TSharedPtr<UE::NNE::IModelNPU> CreateModelNPU(const TObjectPtr<UNNEModelData> ModelData) override { return Impl->CreateModelNPU(ModelData); }
-};
-
 UCLASS()
 class UNNERuntimeORTDml_GPU : public UNNERuntimeORTDmlProxy, public INNERuntimeGPU
 {
@@ -138,26 +125,6 @@ public:
 	virtual TSharedPtr<UE::NNE::IModelGPU> CreateModelGPU(const TObjectPtr<UNNEModelData> ModelData) override { return Impl->CreateModelGPU(ModelData); }
 };
 
-UCLASS()
-class UNNERuntimeORTDml_RDG : public UNNERuntimeORTDmlProxy, public INNERuntimeRDG
-{
-	GENERATED_BODY()
-
-public:
-	virtual ECanCreateModelRDGStatus CanCreateModelRDG(TObjectPtr<UNNEModelData> ModelData) const override { return Impl->CanCreateModelRDG(ModelData); }
-	virtual TSharedPtr<UE::NNE::IModelRDG> CreateModelRDG(TObjectPtr<UNNEModelData> ModelData) override { return Impl->CreateModelRDG(ModelData); }
-};
-
-UCLASS()
-class UNNERuntimeORTDml_NPU : public UNNERuntimeORTDmlProxy, public INNERuntimeNPU
-{
-	GENERATED_BODY()
-
-public:
-	virtual ECanCreateModelNPUStatus CanCreateModelNPU(const TObjectPtr<UNNEModelData> ModelData) const override { return Impl->CanCreateModelNPU(ModelData); }
-	virtual TSharedPtr<UE::NNE::IModelNPU> CreateModelNPU(const TObjectPtr<UNNEModelData> ModelData) override { return Impl->CreateModelNPU(ModelData); }
-};
-
 namespace UE::NNERuntimeORT::Private
 {
 
