[NNE] NNERuntimeORTDml remove NPU interface.

#tests NNEAutomationTests, Win64, Editor

#Synced-CL 50195398

diff --git a/Engine/Plugins/NNE/NNERuntimeORT/Source/NNERuntimeORT/Private/NNERuntimeORT.cpp b/Engine/Plugins/NNE/NNERuntimeORT/Source/NNERuntimeORT/Private/NNERuntimeORT.cpp
--- a/Engine/Plugins/NNE/NNERuntimeORT/Source/NNERuntimeORT/Private/NNERuntimeORT.cpp
+++ b/Engine/Plugins/NNE/NNERuntimeORT/Source/NNERuntimeORT/Private/NNERuntimeORT.cpp
@@ -90,14 +90,6 @@ namespace UE::NNERuntimeORT::Private::Details
 		}();
 		return bIsDirectMLAvailable && bIsRHID3D12Available;
 	}
-	bool IsAvailableNPU(bool bIsDirectMLAvailable)
-	{
-		static bool bIsD3D12DeviceNPUAvailable = []() -> bool
-		{
-			return UE::NNERuntimeORT::Private::IsD3D12DeviceNPUAvailable();
-		}();
-		return bIsDirectMLAvailable && bIsD3D12DeviceNPUAvailable;
-	}
 } // namespace UE::NNERuntimeORT::Private::Details
 
 UNNERuntimeORTCpu::ECanCreateModelDataStatus UNNERuntimeORTCpu::CanCreateModelData(const FString& FileType, TConstArrayView64<uint8> FileData, const TMap<FString, TConstArrayView64<uint8>>& AdditionalFileData, const FGuid& FileId, const ITargetPlatform* TargetPlatform) const
@@ -345,35 +337,6 @@ virtual TSharedPtr<UE::NNE::IModelRDG> CreateModelRDG(TObjectPtr<UNNEModelData>
 #endif // PLATFORM_WINDOWS
 }
 
-virtual ECanCreateModelNPUStatus CanCreateModelNPU(const TObjectPtr<UNNEModelData> ModelData) const override
-{
-	if (!ensureMsgf(UE::NNERuntimeORT::Private::Details::IsAvailableNPU(bIsDirectMLAvailable), TEXT("NPU interface should not be available!")))
-	{
-		return ECanCreateModelCommonStatus::Fail;
-	}
-
-	return CanCreateModelCommon(ModelData) == ECanCreateModelCommonStatus::Ok ? ECanCreateModelNPUStatus::Ok : ECanCreateModelNPUStatus::Fail;
-}
-
-virtual TSharedPtr<UE::NNE::IModelNPU> CreateModelNPU(const TObjectPtr<UNNEModelData> ModelData) override
-{
-#if PLATFORM_WINDOWS
-	check(ModelData);
-
-	if (CanCreateModelNPU(ModelData) != ECanCreateModelRDGStatus::Ok)
-	{
-		UE_LOG(LogNNERuntimeORT, Error, TEXT("Cannot create a model NPU from the model data with id %s"), *ModelData->GetFileId().ToString(EGuidFormats::Digits));
-		return {};
-	}
-
-	const TSharedRef<UE::NNE::FSharedModelData> SharedData = ModelData->GetModelData(GetRuntimeName()).ToSharedRef();
-
-	return MakeShared<UE::NNERuntimeORT::Private::FModelORTNpu>(Environment.ToSharedRef(), SharedData);
-#else // PLATFORM_WINDOWS
-	return {};
-#endif // PLATFORM_WINDOWS
-}
-
 private:
 ECanCreateModelCommonStatus CanCreateModelCommon(const TObjectPtr<UNNEModelData> ModelData, bool bRHID3D12Required = true) const
 {
@@ -438,25 +401,19 @@ TWeakObjectPtr<UNNERuntimeORTDmlProxy> MakeRuntimeDml(bool bDirectMLAvailable)
 	// 
 	// Therefore, at this moment we exceptionally assume D3D12 "is present" and will do the check later in CanCreateModelGPU()
 	const bool bRHID3D12Available = IsRHID3D12Available();
-	const bool bD3D12DeviceNPUAvailable = IsD3D12DeviceNPUAvailable();
 
 	const bool bIsAvailableRDG = bDirectMLAvailable && bRHID3D12Available;
-	const bool bIsAvailableNPU = bDirectMLAvailable && bD3D12DeviceNPUAvailable;
 
 	UE_LOG(LogNNERuntimeORT, Log, TEXT("MakeRuntimeORTDml:"));
 	UE_LOG(LogNNERuntimeORT, Log, TEXT("  DirectML:  %s"), (bDirectMLAvailable ? TEXT("yes") : TEXT("no")));
 	UE_LOG(LogNNERuntimeORT, Log, TEXT("  RHI D3D12: %s"), (bRHID3D12Available ? TEXT("yes") : TEXT("no")));
 	UE_LOG(LogNNERuntimeORT, Log, TEXT("  D3D12:     n/a (checked in CanCreateModelGPU)"));
-	UE_LOG(LogNNERuntimeORT, Log, TEXT("  NPU:       %s"), (bD3D12DeviceNPUAvailable ? TEXT("yes") : TEXT("no")));
 
 	UE_LOG(LogNNERuntimeORT, Log, TEXT("Interface availability:"));
 	UE_LOG(LogNNERuntimeORT, Log, TEXT("  GPU: n/a (checked in CanCreateModelGPU)"));
 	UE_LOG(LogNNERuntimeORT, Log, TEXT("  RDG: %s"), (bIsAvailableRDG ? TEXT("yes") : TEXT("no")));
-	UE_LOG(LogNNERuntimeORT, Log, TEXT("  NPU: %s"), (bIsAvailableNPU ? TEXT("yes") : TEXT("no")));
 
-	if (bIsAvailableRDG && bIsAvailableNPU) return NewObject<UNNERuntimeORTDml_GPU_RDG_NPU>();
-	else if (bIsAvailableRDG) return NewObject<UNNERuntimeORTDml_GPU_RDG>();
-	else if (bIsAvailableNPU) return NewObject<UNNERuntimeORTDml_GPU_NPU>();
+	if (bIsAvailableRDG) return NewObject<UNNERuntimeORTDml_GPU_RDG>();
 	else if (bDirectMLAvailable) return NewObject<UNNERuntimeORTDml_GPU>();
 
 #if WITH_EDITOR
diff --git a/Engine/Plugins/NNE/NNERuntimeORT/Source/NNERuntimeORT/Private/NNERuntimeORT.h b/Engine/Plugins/NNE/NNERuntimeORT/Source/NNERuntimeORT/Private/NNERuntimeORT.h
--- a/Engine/Plugins/NNE/NNERuntimeORT/Source/NNERuntimeORT/Private/NNERuntimeORT.h
+++ b/Engine/Plugins/NNE/NNERuntimeORT/Source/NNERuntimeORT/Private/NNERuntimeORT.h
@@ -4,7 +4,6 @@
 
 #include "NNERuntime.h"
 #include "NNERuntimeCPU.h"
-#include "NNERuntimeNPU.h"
 #include "NNERuntimeGPU.h"
 #include "NNERuntimeRDG.h"
 #include "UObject/Class.h"
@@ -46,7 +45,7 @@ public:
 	virtual TSharedPtr<UE::NNE::IModelCPU> CreateModelCPU(const TObjectPtr<UNNEModelData> ModelData) override;
 };
 
-class INNERuntimeORTDmlImpl : public INNERuntime, public INNERuntimeGPU, public INNERuntimeRDG, public INNERuntimeNPU
+class INNERuntimeORTDmlImpl : public INNERuntime, public INNERuntimeGPU, public INNERuntimeRDG
 {
 public:
 	virtual void Init(TSharedRef<UE::NNERuntimeORT::Private::FEnvironment> InEnvironment, bool bInDirectMLAvailable) = 0;
@@ -73,22 +72,6 @@ protected:
 	TUniquePtr<INNERuntimeORTDmlImpl> Impl;
 };
 
-UCLASS()
-class UNNERuntimeORTDml_GPU_RDG_NPU : public UNNERuntimeORTDmlProxy, public INNERuntimeGPU, public INNERuntimeRDG, public INNERuntimeNPU
-{
-	GENERATED_BODY()
-
-public:
-	virtual ECanCreateModelGPUStatus CanCreateModelGPU(const TObjectPtr<UNNEModelData> ModelData) const override { return Impl->CanCreateModelGPU(ModelData); }
-	virtual TSharedPtr<UE::NNE::IModelGPU> CreateModelGPU(const TObjectPtr<UNNEModelData> ModelData) override { return Impl->CreateModelGPU(ModelData); }
-
-	virtual ECanCreateModelRDGStatus CanCreateModelRDG(TObjectPtr<UNNEModelData> ModelData) const override { return Impl->CanCreateModelRDG(ModelData); }
-	virtual TSharedPtr<UE::NNE::IModelRDG> CreateModelRDG(TObjectPtr<UNNEModelData> ModelData) override { return Impl->CreateModelRDG(ModelData); }
-
-	virtual ECanCreateModelNPUStatus CanCreateModelNPU(const TObjectPtr<UNNEModelData> ModelData) const override { return Impl->CanCreateModelNPU(ModelData); }
-	virtual TSharedPtr<UE::NNE::IModelNPU> CreateModelNPU(const TObjectPtr<UNNEModelData> ModelData) override { return Impl->CreateModelNPU(ModelData); }
-};
-
 UCLASS()
 class UNNERuntimeORTDml_GPU_RDG : public UNNERuntimeORTDmlProxy, public INNERuntimeGPU, public INNERuntimeRDG
 {
@@ -102,19 +85,6 @@ public:
 	virtual TSharedPtr<UE::NNE::IModelRDG> CreateModelRDG(TObjectPtr<UNNEModelData> ModelData) override { return Impl->CreateModelRDG(ModelData); }
 };
 
-UCLASS()
-class UNNERuntimeORTDml_GPU_NPU : public UNNERuntimeORTDmlProxy, public INNERuntimeGPU, public INNERuntimeNPU
-{
-	GENERATED_BODY()
-
-public:
-	virtual ECanCreateModelGPUStatus CanCreateModelGPU(const TObjectPtr<UNNEModelData> ModelData) const override { return Impl->CanCreateModelGPU(ModelData); }
-	virtual TSharedPtr<UE::NNE::IModelGPU> CreateModelGPU(const TObjectPtr<UNNEModelData> ModelData) override { return Impl->CreateModelGPU(ModelData); }
-
-	virtual ECanCreateModelNPUStatus CanCreateModelNPU(const TObjectPtr<UNNEModelData> ModelData) const override { return Impl->CanCreateModelNPU(ModelData); }
-	virtual TSharedPtr<UE::NNE::IModelNPU> CreateModelNPU(const TObjectPtr<UNNEModelData> ModelData) override { return Impl->CreateModelNPU(ModelData); }
-};
-
 UCLASS()
 class UNNERuntimeORTDml_GPU : public UNNERuntimeORTDmlProxy, public INNERuntimeGPU
 {
diff --git a/Engine/Plugins/NNE/NNERuntimeORT/Source/NNERuntimeORT/Private/NNERuntimeORTModel.cpp b/Engine/Plugins/NNE/NNERuntimeORT/Source/NNERuntimeORT/Private/NNERuntimeORTModel.cpp
--- a/Engine/Plugins/NNE/NNERuntimeORT/Source/NNERuntimeORT/Private/NNERuntimeORTModel.cpp
+++ b/Engine/Plugins/NNE/NNERuntimeORT/Source/NNERuntimeORT/Private/NNERuntimeORTModel.cpp
@@ -1158,43 +1158,6 @@ FModelInstanceORTDmlRDG::EEnqueueRDGStatus FModelInstanceORTDmlRDG::EnqueueRDG(F
 
 	return EEnqueueRDGStatus::Ok;
 }
-
-TSharedPtr<NNE::IModelInstanceNPU> FModelORTNpu::CreateModelInstanceNPU()
-{
-	const FRuntimeConf RuntimeConfig = Detail::MakeRuntimeConfigFromSettings(GetDefault<UNNERuntimeORTSettings>());
-
-	TSharedPtr<FModelInstanceORTNpu> ModelInstance = MakeShared<FModelInstanceORTNpu>(RuntimeConfig, Environment);
-	if (!ModelInstance->Init(ModelData->GetView()))
-	{
-		return {};
-	}
-
-	return ModelInstance;
-}
-
-FModelORTNpu::FModelORTNpu(TSharedRef<FEnvironment> InEnvironment, TSharedRef<UE::NNE::FSharedModelData> InModelData) :
-	Environment(InEnvironment), ModelData(InModelData)
-{
-}
-
-bool FModelInstanceORTNpu::InitializedAndConfigureMembers()
-{
-	if (!FModelInstanceORTRunSync::InitializedAndConfigureMembers())
-	{
-		return false;
-	}
-
-	SessionOptions = CreateSessionOptionsForDirectMLNpu(Environment);
-	if (!SessionOptions.IsValid())
-	{
-		return false;
-	}
-
-	SessionOptions->SetExecutionMode(RuntimeConf.ExecutionMode);
-	SessionOptions->SetGraphOptimizationLevel(GetGraphOptimizationLevelForDML(true));
-
-	return true;
-}
 #endif //PLATFORM_WINDOWS
 	
 } // namespace UE::NNERuntimeORT::Private
\ No newline at end of file
diff --git a/Engine/Plugins/NNE/NNERuntimeORT/Source/NNERuntimeORT/Private/NNERuntimeORTModel.h b/Engine/Plugins/NNE/NNERuntimeORT/Source/NNERuntimeORT/Private/NNERuntimeORTModel.h
--- a/Engine/Plugins/NNE/NNERuntimeORT/Source/NNERuntimeORT/Private/NNERuntimeORTModel.h
+++ b/Engine/Plugins/NNE/NNERuntimeORT/Source/NNERuntimeORT/Private/NNERuntimeORTModel.h
@@ -6,7 +6,6 @@
 #include "NNEModelData.h"
 #include "NNEOnnxruntime.h"
 #include "NNERuntimeCPU.h"
-#include "NNERuntimeNPU.h"
 #include "NNERuntimeGPU.h"
 #include "NNERuntimeORTEnv.h"
 #include "NNERuntimeORTTensor.h"
@@ -179,29 +178,6 @@ private:
 	class FProxy; // Contains data written/read by RHI Thread and data initialized by Game/Render Thread
 	TSharedPtr<FProxy> Proxy;
 };
-
-class FModelInstanceORTNpu : public FModelInstanceORTRunSync<NNE::IModelInstanceNPU, NNE::FTensorBindingCPU>
-{
-public:
-	FModelInstanceORTNpu(const FRuntimeConf& InRuntimeConf, TSharedRef<FEnvironment> InEnvironment) : FModelInstanceORTRunSync(InRuntimeConf, InEnvironment) {}
-	virtual ~FModelInstanceORTNpu() = default;
-
-private:
-	virtual bool InitializedAndConfigureMembers() override;
-};
-
-class FModelORTNpu : public NNE::IModelNPU
-{
-public:
-	FModelORTNpu(TSharedRef<FEnvironment> InEnvironment, TSharedRef<UE::NNE::FSharedModelData> InModelData);
-	virtual ~FModelORTNpu() = default;
-
-	virtual TSharedPtr<NNE::IModelInstanceNPU> CreateModelInstanceNPU() override;
-
-private:
-	TSharedRef<FEnvironment> Environment;
-	TSharedRef<UE::NNE::FSharedModelData> ModelData;
-};
 #endif //PLATFORM_WINDOWS
 	
 } // UE::NNERuntimeORT::Private
\ No newline at end of file
diff --git a/Engine/Plugins/NNE/NNERuntimeORT/Source/NNERuntimeORT/Private/NNERuntimeORTUtils.cpp b/Engine/Plugins/NNE/NNERuntimeORT/Source/NNERuntimeORT/Private/NNERuntimeORTUtils.cpp
--- a/Engine/Plugins/NNE/NNERuntimeORT/Source/NNERuntimeORT/Private/NNERuntimeORTUtils.cpp
+++ b/Engine/Plugins/NNE/NNERuntimeORT/Source/NNERuntimeORT/Private/NNERuntimeORTUtils.cpp
@@ -134,233 +134,6 @@ bool IsD3D12Available()
 #endif // PLATFORM_WINDOWS
 }
 
-#if PLATFORM_WINDOWS
-Microsoft::WRL::ComPtr<IUnknown> GetAdapterNpu(bool bVerbose, D3D_FEATURE_LEVEL& OutFeatureLevel)
-{
-	using Microsoft::WRL::ComPtr;
-
-	OutFeatureLevel = D3D_FEATURE_LEVEL_1_0_CORE;
-
-	const bool bForceComputeOnlyDevice = true;
-
-	auto GetAdapterName = [](IDXCoreAdapter* Adapter) -> FString
-		{
-			DXCoreAdapterProperty DriverDescProperty = DXCoreAdapterProperty::DriverDescription;
-
-			if (!Adapter->IsPropertySupported(DriverDescProperty))
-			{
-				return TEXT("Unknown device");
-			}
-
-			size_t PropertySize = 0;
-			HRESULT Hr = Adapter->GetPropertySize(DriverDescProperty, &PropertySize);
-			if (FAILED(Hr))
-			{
-				return TEXT("Unknown device");
-			}
-
-			TArray<char> Buffer;
-			Buffer.Init(0, (int32)PropertySize);
-
-			Hr = Adapter->GetProperty(DriverDescProperty, PropertySize, Buffer.GetData());
-			if (FAILED(Hr))
-			{
-				return TEXT("Unknown device");
-			}
-			
-			return FString(ANSI_TO_TCHAR(Buffer.GetData()));
-		};
-
-	void* DxCoreModule = FPlatformProcess::GetDllHandle(TEXT("DXCore.dll"));
-	if (!DxCoreModule)
-	{
-		UE_LOG(LogNNERuntimeORT, Log, TEXT("Failed to load module 'DXCore.dll'"));
-		return {};
-	}
-
-	using DXCoreCreateAdapterFactoryFunType = HRESULT __stdcall(REFIID, void**);
-
-	DXCoreCreateAdapterFactoryFunType* DxCoreCreateAdapterFactoryFun = reinterpret_cast<DXCoreCreateAdapterFactoryFunType*>(FPlatformProcess::GetDllExport(DxCoreModule, TEXT("DXCoreCreateAdapterFactory")));
-	if (!DxCoreCreateAdapterFactoryFun)
-	{
-		UE_LOG(LogNNERuntimeORT, Log, TEXT("Failed to get export 'DXCoreCreateAdapterFactory' from module 'DXCore.dll'"));
-		return {};
-	}
-
-	ComPtr<IDXCoreAdapterFactory> Factory;
-	HRESULT Hr = DxCoreCreateAdapterFactoryFun(IID_PPV_ARGS(&Factory));
-	if (FAILED(Hr))
-	{
-		UE_LOG(LogNNERuntimeORT, Log, TEXT("Failed to create DXCore Adapter Factory, DxCoreCreateAdapterFactory error code :%x"), Hr);
-		return {};
-	}
-
-	// New path from Windows 11, version 24H2 on
-	{
-#ifdef ORT_USE_NEW_DXCORE_FEATURES
-		const GUID GUID_LIST_GENERIC_ML[] = { DXCORE_ADAPTER_ATTRIBUTE_D3D12_GENERIC_ML };
-#else
-		const GUID GUID_LIST_GENERIC_ML[] = { { 0xb71b0d41, 0x1088, 0x422f, { 0xa2, 0x7c, 0x2, 0x50, 0xb7, 0xd3, 0xa9, 0x88 } } }; // Value of DXCORE_ADAPTER_ATTRIBUTE_D3D12_GENERIC_ML
-#endif
-		const GUID GUID_LIST_CORE_COMPUTE[] = { DXCORE_ADAPTER_ATTRIBUTE_D3D12_CORE_COMPUTE };
-
-		auto GetNonGraphicsAdapter = [] (IDXCoreAdapterList* AdapterList, IDXCoreAdapter** OutAdapter)
-		{
-			for (uint32_t i = 0, AdapterCount = AdapterList->GetAdapterCount(); i < AdapterCount; i++)
-			{
-				ComPtr<IDXCoreAdapter> Adapter;
-				HRESULT Hr = AdapterList->GetAdapter(i, IID_PPV_ARGS(&Adapter));
-				if (FAILED(Hr))
-				{
-					UE_LOG(LogNNERuntimeORT, Warning, TEXT("%d: Failed to get adapter, IDXCoreAdapterList::GetAdapter error code: %x"), i, Hr);
-					continue;
-				}
-
-				if (!Adapter->IsAttributeSupported(DXCORE_ADAPTER_ATTRIBUTE_D3D12_GRAPHICS))
-				{
-					*OutAdapter = Adapter.Detach();
-					return;
-				}
-			}
-			*OutAdapter = nullptr;
-		};
-
-		ComPtr<IDXCoreAdapter> Adapter;
-		ComPtr<IDXCoreAdapterList> AdapterList;
-		D3D_FEATURE_LEVEL FeatureLevel = D3D_FEATURE_LEVEL_1_0_GENERIC;
-
-		Hr = Factory->CreateAdapterList(ARRAYSIZE(GUID_LIST_GENERIC_ML), GUID_LIST_GENERIC_ML, IID_PPV_ARGS(&AdapterList));
-		if (FAILED(Hr))
-		{
-			UE_LOG(LogNNERuntimeORT, Log, TEXT("Failed to create DXCore Adapter List, IDXCoreAdapterFactory::CreateAdapterList error code :%x"), Hr);
-			return {};
-		}
-
-		if (AdapterList->GetAdapterCount() > 0)
-		{
-			GetNonGraphicsAdapter(AdapterList.Get(), Adapter.GetAddressOf());
-		}
-
-		if (!Adapter)
-		{
-			FeatureLevel = D3D_FEATURE_LEVEL_1_0_CORE;
-
-			Hr = Factory->CreateAdapterList(ARRAYSIZE(GUID_LIST_CORE_COMPUTE), GUID_LIST_CORE_COMPUTE, IID_PPV_ARGS(&AdapterList));
-			if (FAILED(Hr))
-			{
-				UE_LOG(LogNNERuntimeORT, Log, TEXT("Failed to create DXCore Adapter List, IDXCoreAdapterFactory::CreateAdapterList error code :%x"), Hr);
-				return {};
-			}
-
-			GetNonGraphicsAdapter(AdapterList.Get(), Adapter.GetAddressOf());
-		}
-
-		if (!Adapter)
-		{
-			UE_LOG(LogNNERuntimeORT, Log, TEXT("No NPU adapter found with attributes (!GRAPHICS && (GENERIC_ML || CORE_COMPUTE)) (Windows 11 Version 24H2 or newer)!"));
-		}
-		else
-		{
-			if (bVerbose)
-			{
-				UE_LOG(LogNNERuntimeORT, Log, TEXT("Selecting NPU adapter: %s"), *(GetAdapterName(Adapter.Get())));
-			}
-
-			OutFeatureLevel = FeatureLevel;
-
-			return Adapter;
-		}
-	}
-
-	const GUID DxGUIDs[] = { DXCORE_ADAPTER_ATTRIBUTE_D3D12_CORE_COMPUTE };
-
-	ComPtr<IDXCoreAdapterList> AdapterList;
-	Hr = Factory->CreateAdapterList(ARRAYSIZE(DxGUIDs), DxGUIDs, IID_PPV_ARGS(&AdapterList));
-	if (FAILED(Hr))
-	{
-		UE_LOG(LogNNERuntimeORT, Log, TEXT("Failed to create DXCore Adapter List, IDXCoreAdapterFactory::CreateAdapterList error code :%x"), Hr);
-		return {};
-	}
-
-	if (bVerbose)
-	{
-		UE_LOG(LogNNERuntimeORT, Log, TEXT("Available graphics and compute adapters:"));
-	}
-
-	ComPtr<IDXCoreAdapter> Adapter;
-	for (uint32_t i = 0, AdapterCount = AdapterList->GetAdapterCount(); i < AdapterCount; i++)
-	{
-		ComPtr<IDXCoreAdapter> CurrentAdapter;
-		Hr = AdapterList->GetAdapter(static_cast<uint32_t>(i), IID_PPV_ARGS(&CurrentAdapter));
-		if (FAILED(Hr))
-		{
-			UE_LOG(LogNNERuntimeORT, Warning, TEXT("%d: Failed to get adapter, IDXCoreAdapterList::GetAdapter error code: %x"), i, Hr);
-			continue;
-		}
-
-		const bool bIsGraphicsAdapter = CurrentAdapter->IsAttributeSupported(DXCORE_ADAPTER_ATTRIBUTE_D3D12_GRAPHICS);
-
-		if (bVerbose)
-		{
-			TArray<FString> SupportedAttributesStringList;
-			SupportedAttributesStringList.Add("Compute");
-			if (bIsGraphicsAdapter)
-			{
-				SupportedAttributesStringList.Add("Graphics");
-			}
-
-			FString SupportedAttributesString = FString::Join(SupportedAttributesStringList, TEXT(", "));
-
-			UE_LOG(LogNNERuntimeORT, Log, TEXT("%d: %s (%s)"), i, *(GetAdapterName(CurrentAdapter.Get())), *SupportedAttributesString);
-		}
-
-		if (!Adapter && (!bForceComputeOnlyDevice || !bIsGraphicsAdapter))
-		{
-			Adapter = std::move(CurrentAdapter);
-		}
-	}
-
-	if (!Adapter)
-	{
-		UE_LOG(LogNNERuntimeORT, Log, TEXT("No NPU adapter found!"));
-
-		return {};
-	}
-
-	if (bVerbose)
-	{
-		UE_LOG(LogNNERuntimeORT, Log, TEXT("Selecting NPU adapter: %s"), *(GetAdapterName(Adapter.Get())));
-	}
-
-	return Adapter;
-}
-#endif // PLATFORM_WINDOWS
-
-bool IsD3D12DeviceNPUAvailable()
-{
-#if PLATFORM_WINDOWS
-	using Microsoft::WRL::ComPtr;
-
-	D3D_FEATURE_LEVEL FeatureLevel = D3D_FEATURE_LEVEL_1_0_CORE;
-
-	ComPtr<IUnknown> AdapterNpu = GetAdapterNpu(true, FeatureLevel);
-	if (!AdapterNpu)
-	{
-		return false;
-	}
-
-	ComPtr<ID3D12Device1> D3D12Device = CreateD3D12Device(AdapterNpu.Get(), FeatureLevel);
-	if (!D3D12Device)
-	{
-		return false;
-	}
-
-	return true;
-#else
-	return false;
-#endif // PLATFORM_WINDOWS
-}
-
 // For more details about ORT graph optimization checkout
 // https://onnxruntime.ai/docs/performance/model-optimizations/graph-optimizations.html
 
@@ -596,85 +369,6 @@ TUniquePtr<Ort::SessionOptions> CreateSessionOptionsForDirectML(const TSharedRef
 #endif // PLATFORM_WINDOWS
 }
 
-TUniquePtr<Ort::SessionOptions> CreateSessionOptionsForDirectMLNpu(const TSharedRef<FEnvironment>& Environment)
-{
-#if PLATFORM_WINDOWS
-	using Microsoft::WRL::ComPtr;
-
-	TUniquePtr<Ort::SessionOptions> SessionOptions = CreateSessionOptionsDefault(Environment);
-	if (!SessionOptions.IsValid())
-	{
-		return {};
-	}
-
-	// Configure for DirectML
-	SessionOptions->SetExecutionMode(ExecutionMode::ORT_SEQUENTIAL);
-	SessionOptions->DisableMemPattern();
-
-	D3D_FEATURE_LEVEL FeatureLevel = D3D_FEATURE_LEVEL_1_0_CORE;
-
-	ComPtr<IUnknown> AdapterNpu = GetAdapterNpu(false, FeatureLevel);
-	if (!AdapterNpu)
-	{
-		UE_LOG(LogNNERuntimeORT, Error, TEXT("Failed to get NPU adapter."));
-		return {};
-	}
-
-	ComPtr<ID3D12Device1> D3D12DeviceNpu = CreateD3D12Device(AdapterNpu.Get(), FeatureLevel);
-	if (!D3D12DeviceNpu)
-	{
-		UE_LOG(LogNNERuntimeORT, Error, TEXT("Failed to create NPU device."));
-		return {};
-	}
-
-	DML_CREATE_DEVICE_FLAGS DmlCreateFlags = DML_CREATE_DEVICE_FLAG_NONE;
-
-	// Set debugging flags
-	if (GRHIGlobals.IsDebugLayerEnabled)
-	{
-		DmlCreateFlags |= DML_CREATE_DEVICE_FLAG_DEBUG;
-	}
-
-	ComPtr<IDMLDevice> DmlDeviceNpu;
-	HRESULT Hr = DMLCreateDevice(D3D12DeviceNpu.Get(), DmlCreateFlags, IID_PPV_ARGS(DmlDeviceNpu.GetAddressOf()));
-	if (FAILED(Hr))
-	{
-		UE_LOG(LogNNERuntimeORT, Error, TEXT("Failed to create DirectML device, DMLCreateDevice error code :%x"), Hr);
-		return {};
-	}
-
-	D3D12_COMMAND_QUEUE_DESC CommandQueueDesc{};
-	CommandQueueDesc.Type = D3D12_COMMAND_LIST_TYPE::D3D12_COMMAND_LIST_TYPE_COMPUTE;
-
-	ComPtr<ID3D12CommandQueue> CommandQueue;
-	Hr = D3D12DeviceNpu->CreateCommandQueue(&CommandQueueDesc, IID_PPV_ARGS(CommandQueue.GetAddressOf()));
-	if (FAILED(Hr))
-	{
-		UE_LOG(LogNNERuntimeORT, Error, TEXT("Failed to create DirectML Command Queue, CreateCommandQueue error code :%x"), Hr);
-		return {};
-	}
-
-	const OrtDmlApi* DmlApi = nullptr;
-	Ort::ThrowOnError(Ort::GetApi().GetExecutionProviderApi("DML", ORT_API_VERSION, reinterpret_cast<const void**>(&DmlApi)));
-	if (!DmlApi)
-	{
-		UE_LOG(LogNNERuntimeORT, Error, TEXT("Ort DirectML Api not available!"));
-		return {};
-	}
-
-	OrtStatusPtr Status = DmlApi->SessionOptionsAppendExecutionProvider_DML1(*SessionOptions, DmlDeviceNpu.Get(), CommandQueue.Get());
-	if (Status)
-	{
-		UE_LOG(LogNNERuntimeORT, Error, TEXT("Failed to add DirectML execution provider to OnnxRuntime session options: %s"), ANSI_TO_TCHAR(Ort::GetApi().GetErrorMessage(Status)));
-		return {};
-	}
-
-	return SessionOptions;
-#else
-	return {};
-#endif // PLATFORM_WINDOWS
-}
-
 bool OptimizeModel(const TSharedRef<FEnvironment> &Environment, Ort::SessionOptions &SessionOptions, 
 					TConstArrayView64<uint8>& InputModel, TArray64<uint8>& OptimizedModel)
 {
diff --git a/Engine/Plugins/NNE/NNERuntimeORT/Source/NNERuntimeORT/Private/NNERuntimeORTUtils.h b/Engine/Plugins/NNE/NNERuntimeORT/Source/NNERuntimeORT/Private/NNERuntimeORTUtils.h
--- a/Engine/Plugins/NNE/NNERuntimeORT/Source/NNERuntimeORT/Private/NNERuntimeORTUtils.h
+++ b/Engine/Plugins/NNE/NNERuntimeORT/Source/NNERuntimeORT/Private/NNERuntimeORTUtils.h
@@ -16,8 +16,6 @@ namespace UE::NNERuntimeORT::Private
 
 	bool IsD3D12Available();
 
-	bool IsD3D12DeviceNPUAvailable();
-
 namespace OrtHelper
 {
 	TArray<uint32> GetShape(const Ort::Value& OrtTensor);
@@ -31,8 +29,6 @@ namespace OrtHelper
 
 	TUniquePtr<Ort::SessionOptions> CreateSessionOptionsForDirectML(const TSharedRef<FEnvironment> &Environment, bool bRHID3D12Required = true);
 
-	TUniquePtr<Ort::SessionOptions> CreateSessionOptionsForDirectMLNpu(const TSharedRef<FEnvironment>& Environment);
-
 	bool OptimizeModel(const TSharedRef<FEnvironment>& Environment, Ort::SessionOptions& SessionOptions, 
 		TConstArrayView64<uint8>& InputModel, TArray64<uint8>& OptimizedModel);
 
